Add new template sequence feature.

Template sequences are sequences that are bound to a single asset archetype or class.They can then either be played as a clip on normal level sequence binding (possessable or spawnable), or played "standalone", in which case the template sequence will played on a spawned object, or on an object provided as a binding override in Blueprint.

diff --git a/Engine/Config/BaseEditorPerProjectUserSettings.ini b/Engine/Config/BaseEditorPerProjectUserSettings.ini
--- a/Engine/Config/BaseEditorPerProjectUserSettings.ini
+++ b/Engine/Config/BaseEditorPerProjectUserSettings.ini
@@ -848,6 +848,15 @@ bInfiniteKeyAreas=true
 bAutoSetTrackDefaults=true
 FrameNumberDisplayFormat=Frames
 
+[TemplateSequenceEditor SequencerSettings]
+bKeyInterpPropertiesOnly=true
+bShowRangeSlider=true
+bKeepPlayRangeInSectionBounds=false
+ZeroPadFrames=4
+bInfiniteKeyAreas=true
+bAutoSetTrackDefaults=true
+FrameNumberDisplayFormat=Frames
+
 [TakeRecorderSequenceEditor SequencerSettings]
 bKeyInterpPropertiesOnly=true
 bShowRangeSlider=true
diff --git a/Engine/Plugins/MovieScene/TemplateSequence/Source/TemplateSequence/Private/Evaluation/TemplateSequenceInstanceData.cpp b/Engine/Plugins/MovieScene/TemplateSequence/Source/TemplateSequence/Private/Evaluation/TemplateSequenceInstanceData.cpp
new file mode 100644
--- /dev/null
+++ b/Engine/Plugins/MovieScene/TemplateSequence/Source/TemplateSequence/Private/Evaluation/TemplateSequenceInstanceData.cpp
@@ -0,0 +1,3 @@
+// Copyright 1998-2019 Epic Games, Inc. All Rights Reserved.
+
+#include "Evaluation/TemplateSequenceInstanceData.h"
diff --git a/Engine/Plugins/MovieScene/TemplateSequence/Source/TemplateSequence/Private/Evaluation/TemplateSequenceSectionTemplate.cpp b/Engine/Plugins/MovieScene/TemplateSequence/Source/TemplateSequence/Private/Evaluation/TemplateSequenceSectionTemplate.cpp
new file mode 100644
--- /dev/null
+++ b/Engine/Plugins/MovieScene/TemplateSequence/Source/TemplateSequence/Private/Evaluation/TemplateSequenceSectionTemplate.cpp
@@ -0,0 +1,52 @@
+// Copyright 1998-2019 Epic Games, Inc. All Rights Reserved.
+
+#include "Evaluation/TemplateSequenceSectionTemplate.h"
+#include "IMovieScenePlaybackClient.h"
+#include "TemplateSequence.h"
+#include "Sections/TemplateSequenceSection.h"
+
+FTemplateSequenceSectionTemplate::FTemplateSequenceSectionTemplate()
+{
+}
+
+FTemplateSequenceSectionTemplate::FTemplateSequenceSectionTemplate(const UTemplateSequenceSection& Section)
+	: SectionStartTime(Section.HasStartFrame() ? Section.GetInclusiveStartFrame() : 0)
+{
+	if (const UTemplateSequence* InnerSequence = Cast<UTemplateSequence>(Section.GetSequence()))
+	{
+		if (InnerSequence->GetMovieScene() != nullptr)
+		{
+			const TArray<FMovieSceneBinding>& InnerSequenceBindings = InnerSequence->GetMovieScene()->GetBindings();
+			if (InnerSequenceBindings.Num() > 0)
+			{
+				const FMovieSceneBinding& Binding = InnerSequenceBindings[0];
+				InnerOperand = FMovieSceneEvaluationOperand(Section.GetSequenceID(), Binding.GetObjectGuid());
+			}
+		}
+	}
+}
+
+void FTemplateSequenceSectionTemplate::SetupOverrides()
+{
+	EnableOverrides(RequiresInitializeFlag | RequiresTearDownFlag);
+}
+
+void FTemplateSequenceSectionTemplate::Evaluate(const FMovieSceneEvaluationOperand& Operand, const FMovieSceneContext& Context, const FPersistentEvaluationData& PersistentData, FMovieSceneExecutionTokens& ExecutionTokens) const
+{
+}
+
+void FTemplateSequenceSectionTemplate::Initialize(const FMovieSceneEvaluationOperand& Operand, const FMovieSceneContext& Context, FPersistentEvaluationData& PersistentData, IMovieScenePlayer& Player) const
+{
+	if (InnerOperand.IsValid())
+	{
+		Player.BindingOverrides.Add(InnerOperand, Operand);
+	}
+}
+
+void FTemplateSequenceSectionTemplate::TearDown(FPersistentEvaluationData& PersistentData, IMovieScenePlayer& Player) const
+{
+	if (InnerOperand.IsValid())
+	{
+		Player.BindingOverrides.Remove(InnerOperand);
+	}
+}
diff --git a/Engine/Plugins/MovieScene/TemplateSequence/Source/TemplateSequence/Private/Section/TemplateSequenceSection.cpp b/Engine/Plugins/MovieScene/TemplateSequence/Source/TemplateSequence/Private/Section/TemplateSequenceSection.cpp
new file mode 100644
--- /dev/null
+++ b/Engine/Plugins/MovieScene/TemplateSequence/Source/TemplateSequence/Private/Section/TemplateSequenceSection.cpp
@@ -0,0 +1,39 @@
+// Copyright 1998-2019 Epic Games, Inc. All Rights Reserved.
+
+#include "Sections/TemplateSequenceSection.h"
+#include "TemplateSequence.h"
+#include "Evaluation/TemplateSequenceInstanceData.h"
+#include "Evaluation/TemplateSequenceSectionTemplate.h"
+#include "UObject/SequencerObjectVersion.h"
+
+
+UTemplateSequenceSection::UTemplateSequenceSection()
+{
+}
+
+void UTemplateSequenceSection::OnDilated(float DilationFactor, FFrameNumber Origin)
+{
+	// TODO-lchabant: shouldn't this be in the base class?
+	Parameters.TimeScale /= DilationFactor;
+}
+
+FMovieSceneSubSequenceData UTemplateSequenceSection::GenerateSubSequenceData(const FSubSequenceInstanceDataParams& Params) const
+{
+	FMovieSceneSubSequenceData SubData(*this);
+
+	FTemplateSequenceInstanceData InstanceData;
+	InstanceData.Operand = Params.Operand;
+	SubData.InstanceData = FMovieSceneSequenceInstanceDataPtr(InstanceData);
+
+	return SubData;
+}
+
+FMovieSceneEvalTemplatePtr UTemplateSequenceSection::GenerateTemplate() const
+{
+	if (SubSequence != nullptr)
+	{
+		return FTemplateSequenceSectionTemplate(*this);
+	}
+	return FMovieSceneEvalTemplatePtr();
+}
+
diff --git a/Engine/Plugins/MovieScene/TemplateSequence/Source/TemplateSequence/Private/TemplateSequence.cpp b/Engine/Plugins/MovieScene/TemplateSequence/Source/TemplateSequence/Private/TemplateSequence.cpp
new file mode 100644
--- /dev/null
+++ b/Engine/Plugins/MovieScene/TemplateSequence/Source/TemplateSequence/Private/TemplateSequence.cpp
@@ -0,0 +1,104 @@
+// Copyright 1998-2019 Epic Games, Inc. All Rights Reserved.
+
+#include "TemplateSequence.h"
+#include "MovieScene.h"
+#include "Components/ActorComponent.h"
+#include "GameFramework/Actor.h"
+
+IMPLEMENT_MODULE(FDefaultModuleImpl, TemplateSequence);
+
+UTemplateSequence::UTemplateSequence(const FObjectInitializer& ObjectInitializer)
+	: Super(ObjectInitializer)
+	, MovieScene(nullptr)
+{
+	bParentContextsAreSignificant = true;
+}
+
+void UTemplateSequence::Initialize()
+{
+	MovieScene = NewObject<UMovieScene>(this, NAME_None, RF_Transactional);
+
+	MovieScene->SetEvaluationType(EMovieSceneEvaluationType::WithSubFrames);
+
+	FFrameRate TickResolution(24000, 1);
+	MovieScene->SetTickResolutionDirectly(TickResolution);
+
+	FFrameRate DisplayRate(30, 1);
+	MovieScene->SetDisplayRate(DisplayRate);
+}
+
+void UTemplateSequence::BindPossessableObject(const FGuid& ObjectId, UObject& PossessedObject, UObject* Context)
+{
+	if (UActorComponent* Component = Cast<UActorComponent>(&PossessedObject))
+	{
+		const FName ComponentName = Component->GetFName();
+		BoundActorComponents.Add(ObjectId, ComponentName);
+	}
+}
+
+bool UTemplateSequence::CanPossessObject(UObject& Object, UObject* InPlaybackContext) const
+{
+	return Object.IsA<AActor>() || Object.IsA<UActorComponent>();
+}
+
+void UTemplateSequence::LocateBoundObjects(const FGuid& ObjectId, UObject* Context, TArray<UObject*, TInlineAllocator<1>>& OutObjects) const
+{
+	AActor* Actor = Cast<AActor>(Context);
+	if (Actor == nullptr)
+	{
+		return;
+	}
+
+	const FName* ComponentName = BoundActorComponents.Find(ObjectId);
+	if (ComponentName == nullptr)
+	{
+		return;
+	}
+
+	if (UActorComponent* FoundComponent = FindObject<UActorComponent>(Actor, *ComponentName->ToString(), false))
+	{
+		OutObjects.Add(FoundComponent);
+	}
+}
+
+UMovieScene* UTemplateSequence::GetMovieScene() const
+{
+	return MovieScene;
+}
+
+UObject* UTemplateSequence::GetParentObject(UObject* Object) const
+{
+	if (UActorComponent* Component = Cast<UActorComponent>(Object))
+	{
+		return Component->GetOwner();
+	}
+
+	return nullptr;
+}
+
+void UTemplateSequence::UnbindPossessableObjects(const FGuid& ObjectId)
+{
+	BoundActorComponents.Remove(ObjectId);
+}
+
+void UTemplateSequence::UnbindObjects(const FGuid& ObjectId, const TArray<UObject*>& InObjects, UObject* Context)
+{
+	BoundActorComponents.Remove(ObjectId);
+}
+
+void UTemplateSequence::UnbindInvalidObjects(const FGuid& ObjectId, UObject* Context)
+{
+	BoundActorComponents.Remove(ObjectId);
+}
+
+bool UTemplateSequence::AllowsSpawnableObjects() const
+{
+	return true;
+}
+
+#if WITH_EDITOR
+FText UTemplateSequence::GetDisplayName() const
+{
+	return UMovieSceneSequence::GetDisplayName();
+}
+#endif
diff --git a/Engine/Plugins/MovieScene/TemplateSequence/Source/TemplateSequence/Private/TemplateSequenceActor.cpp b/Engine/Plugins/MovieScene/TemplateSequence/Source/TemplateSequence/Private/TemplateSequenceActor.cpp
new file mode 100644
--- /dev/null
+++ b/Engine/Plugins/MovieScene/TemplateSequence/Source/TemplateSequence/Private/TemplateSequenceActor.cpp
@@ -0,0 +1,198 @@
+// Copyright 1998-2019 Epic Games, Inc. All Rights Reserved.
+
+#include "TemplateSequenceActor.h"
+#include "TemplateSequence.h"
+#include "Engine/World.h"
+#include "Logging/MessageLog.h"
+#include "Misc/UObjectToken.h"
+#include "Net/UnrealNetwork.h"
+
+ATemplateSequenceActor::ATemplateSequenceActor(const FObjectInitializer& ObjectInitializer)
+	: Super(ObjectInitializer)
+{
+	USceneComponent* SceneComponent = CreateDefaultSubobject<USceneComponent>(TEXT("SceneComp"));
+	RootComponent = SceneComponent;
+
+	SequencePlayer = ObjectInitializer.CreateDefaultSubobject<UTemplateSequencePlayer>(this, "AnimationPlayer");
+
+	// Note: we don't set `bCanEverTick` because we add this actor to the list of level sequence actors in the world.
+	//PrimaryActorTick.bCanEverTick = false;
+}
+
+void ATemplateSequenceActor::PostInitProperties()
+{
+	Super::PostInitProperties();
+
+	// Have to initialize this here as any properties set on default subobjects inside the constructor
+	// Get stomped by the CDO's properties when the constructor exits.
+	SequencePlayer->SetPlaybackClient(this);
+}
+
+void ATemplateSequenceActor::GetLifetimeReplicatedProps(TArray<FLifetimeProperty>& OutLifetimeProps) const
+{
+	Super::GetLifetimeReplicatedProps(OutLifetimeProps);
+
+	DOREPLIFETIME(ATemplateSequenceActor, SequencePlayer);
+}
+
+UTemplateSequencePlayer* ATemplateSequenceActor::GetSequencePlayer() const
+{
+	return SequencePlayer && SequencePlayer->GetSequence() ? SequencePlayer : nullptr;
+}
+
+void ATemplateSequenceActor::PostInitializeComponents()
+{
+	Super::PostInitializeComponents();
+
+	InitializePlayer();
+}
+
+void ATemplateSequenceActor::BeginPlay()
+{
+	GetWorld()->LevelSequenceActors.Add(this);
+
+	Super::BeginPlay();
+	
+	if (PlaybackSettings.bAutoPlay)
+	{
+		SequencePlayer->Play();
+	}
+}
+
+void ATemplateSequenceActor::EndPlay(const EEndPlayReason::Type EndPlayReason)
+{
+	if (SequencePlayer)
+	{
+		// See comment in LevelSequenceActor.cpp
+		SequencePlayer->Stop();
+	}
+
+	GetWorld()->LevelSequenceActors.Remove(this);
+
+	Super::EndPlay(EndPlayReason);
+}
+
+void ATemplateSequenceActor::Tick(float DeltaSeconds)
+{
+	Super::Tick(DeltaSeconds);
+
+	if (SequencePlayer)
+	{
+		SequencePlayer->Update(DeltaSeconds);
+	}
+}
+
+UTemplateSequence* ATemplateSequenceActor::GetSequence() const
+{
+	return Cast<UTemplateSequence>(TemplateSequence.ResolveObject());
+}
+
+UTemplateSequence* ATemplateSequenceActor::LoadSequence() const
+{
+	return Cast<UTemplateSequence>(TemplateSequence.TryLoad());
+}
+
+void ATemplateSequenceActor::SetSequence(UTemplateSequence* InSequence)
+{
+	if (!SequencePlayer->IsPlaying())
+	{
+		TemplateSequence = InSequence;
+
+		if (InSequence)
+		{
+			SequencePlayer->Initialize(InSequence, GetWorld(), PlaybackSettings);
+		}
+	}
+}
+
+void ATemplateSequenceActor::InitializePlayer()
+{
+	if (TemplateSequence.IsValid() && GetWorld()->IsGameWorld())
+	{
+		// Attempt to resolve the asset without loading it
+		UTemplateSequence* TemplateSequenceAsset = GetSequence();
+		if (TemplateSequenceAsset)
+		{
+			// Level sequence is already loaded. Initialize the player if it's not already initialized with this sequence
+			if (TemplateSequenceAsset != SequencePlayer->GetSequence())
+			{
+				SequencePlayer->Initialize(TemplateSequenceAsset, GetWorld(), PlaybackSettings);
+			}
+		}
+		else if (!IsAsyncLoading())
+		{
+			TemplateSequenceAsset = LoadSequence();
+			if (TemplateSequenceAsset != SequencePlayer->GetSequence())
+			{
+				SequencePlayer->Initialize(TemplateSequenceAsset, GetWorld(), PlaybackSettings);
+			}
+		}
+		else
+		{
+			LoadPackageAsync(TemplateSequence.GetLongPackageName(), FLoadPackageAsyncDelegate::CreateUObject(this, &ATemplateSequenceActor::OnSequenceLoaded));
+		}
+	}
+}
+
+void ATemplateSequenceActor::OnSequenceLoaded(const FName& PackageName, UPackage* Package, EAsyncLoadingResult::Type Result)
+{
+	if (Result == EAsyncLoadingResult::Succeeded)
+	{
+		UTemplateSequence* TemplateSequenceAsset = GetSequence();
+		if (SequencePlayer && SequencePlayer->GetSequence() != TemplateSequenceAsset)
+		{
+			SequencePlayer->Initialize(TemplateSequenceAsset, GetWorld(), PlaybackSettings);
+		}
+	}
+}
+
+bool ATemplateSequenceActor::RetrieveBindingOverrides(const FGuid& InBindingId, FMovieSceneSequenceID InSequenceID, TArray<UObject*, TInlineAllocator<1>>& OutObjects) const
+{
+	if (UObject* Object = BindingOverride.Object.Get())
+	{
+		OutObjects.Add(Object);
+		return false;
+	}
+
+	// No binding overrides, use default binding.
+	return true;
+}
+
+UObject* ATemplateSequenceActor::GetInstanceData() const
+{
+	return nullptr;
+}
+
+void ATemplateSequenceActor::SetBinding(AActor* Actor)
+{
+	BindingOverride.Object = Actor;
+
+	UTemplateSequence* TemplateSequenceObject = GetSequence();
+	if (SequencePlayer && TemplateSequenceObject)
+	{
+		UMovieScene* MovieScene = TemplateSequenceObject->GetMovieScene();
+		if (MovieScene != nullptr && MovieScene->GetBindings().Num() > 0)
+		{
+			const FMovieSceneBinding& Binding = MovieScene->GetBindings()[0];
+			SequencePlayer->State.Invalidate(Binding.GetObjectGuid(), MovieSceneSequenceID::Root);
+		}
+	}
+}
+
+#if WITH_EDITOR
+
+bool ATemplateSequenceActor::GetReferencedContentObjects(TArray<UObject*>& Objects) const
+{
+	UTemplateSequence* TemplateSequenceAsset = LoadSequence();
+
+	if (TemplateSequenceAsset)
+	{
+		Objects.Add(TemplateSequenceAsset);
+	}
+
+	Super::GetReferencedContentObjects(Objects);
+
+	return true;
+}
+
+#endif // WITH_EDITOR
diff --git a/Engine/Plugins/MovieScene/TemplateSequence/Source/TemplateSequence/Private/TemplateSequencePlayer.cpp b/Engine/Plugins/MovieScene/TemplateSequence/Source/TemplateSequence/Private/TemplateSequencePlayer.cpp
new file mode 100644
--- /dev/null
+++ b/Engine/Plugins/MovieScene/TemplateSequence/Source/TemplateSequence/Private/TemplateSequencePlayer.cpp
@@ -0,0 +1,61 @@
+// Copyright 1998-2019 Epic Games, Inc. All Rights Reserved.
+
+#include "TemplateSequencePlayer.h"
+#include "TemplateSequence.h"
+#include "TemplateSequenceActor.h"
+#include "TemplateSequenceSpawnRegister.h"
+#include "Engine/Engine.h"
+
+UTemplateSequencePlayer::UTemplateSequencePlayer(const FObjectInitializer& ObjectInitializer)
+	: Super(ObjectInitializer)
+{
+}
+
+UTemplateSequencePlayer* UTemplateSequencePlayer::CreateTemplateSequencePlayer(UObject* WorldContextObject, UTemplateSequence* TemplateSequence, FMovieSceneSequencePlaybackSettings Settings, ATemplateSequenceActor*& OutActor)
+{
+	if (TemplateSequence == nullptr)
+	{
+		return nullptr;
+	}
+
+	UWorld* World = GEngine->GetWorldFromContextObject(WorldContextObject, EGetWorldErrorMode::LogAndReturnNull);
+	if (World == nullptr)
+	{
+		return nullptr;
+	}
+
+	FActorSpawnParameters SpawnParams;
+	SpawnParams.SpawnCollisionHandlingOverride = ESpawnActorCollisionHandlingMethod::AlwaysSpawn;
+	SpawnParams.ObjectFlags |= RF_Transient;
+	SpawnParams.bAllowDuringConstructionScript = true;
+
+	// Defer construction for autoplay so that BeginPlay() is called
+	SpawnParams.bDeferConstruction = true;
+
+	ATemplateSequenceActor* Actor = World->SpawnActor<ATemplateSequenceActor>(SpawnParams);
+
+	Actor->PlaybackSettings = Settings;
+	Actor->TemplateSequence = TemplateSequence;
+
+	Actor->InitializePlayer();
+	OutActor = Actor;
+
+	FTransform DefaultTransform;
+	Actor->FinishSpawning(DefaultTransform);
+
+	return Actor->SequencePlayer;
+}
+
+void UTemplateSequencePlayer::Initialize(UMovieSceneSequence* InSequence, UWorld* InWorld, const FMovieSceneSequencePlaybackSettings& InSettings)
+{
+	SpawnRegister = MakeShareable(new FTemplateSequenceSpawnRegister());
+
+	World = InWorld;
+
+	Super::Initialize(InSequence, InSettings);
+}
+
+UObject* UTemplateSequencePlayer::GetPlaybackContext() const
+{
+	return World.Get();
+}
diff --git a/Engine/Plugins/MovieScene/TemplateSequence/Source/TemplateSequence/Private/TemplateSequenceSpawnRegister.cpp b/Engine/Plugins/MovieScene/TemplateSequence/Source/TemplateSequence/Private/TemplateSequenceSpawnRegister.cpp
new file mode 100644
--- /dev/null
+++ b/Engine/Plugins/MovieScene/TemplateSequence/Source/TemplateSequence/Private/TemplateSequenceSpawnRegister.cpp
@@ -0,0 +1,7 @@
+// Copyright 1998-2019 Epic Games, Inc. All Rights Reserved.
+
+#include "TemplateSequenceSpawnRegister.h"
+
+FTemplateSequenceSpawnRegister::FTemplateSequenceSpawnRegister()
+{
+}
diff --git a/Engine/Plugins/MovieScene/TemplateSequence/Source/TemplateSequence/Private/Tracks/TemplateSequenceTrack.cpp b/Engine/Plugins/MovieScene/TemplateSequence/Source/TemplateSequence/Private/Tracks/TemplateSequenceTrack.cpp
new file mode 100644
--- /dev/null
+++ b/Engine/Plugins/MovieScene/TemplateSequence/Source/TemplateSequence/Private/Tracks/TemplateSequenceTrack.cpp
@@ -0,0 +1,58 @@
+// Copyright 1998-2019 Epic Games, Inc. All Rights Reserved.
+
+#include "Tracks/TemplateSequenceTrack.h"
+#include "IMovieSceneTracksModule.h"
+#include "TemplateSequence.h"
+#include "MovieSceneTimeHelpers.h"
+#include "Sections/TemplateSequenceSection.h"
+#include "Compilation/IMovieSceneTemplateGenerator.h"
+#include "Evaluation/TemplateSequenceSectionTemplate.h"
+
+#define LOCTEXT_NAMESPACE "TemplateSequenceTrack"
+
+bool UTemplateSequenceTrack::SupportsType(TSubclassOf<UMovieSceneSection> SectionClass) const
+{
+	return SectionClass == UTemplateSequenceSection::StaticClass();
+}
+
+UMovieSceneSection* UTemplateSequenceTrack::CreateNewSection()
+{
+	return NewObject<UTemplateSequenceSection>(this, NAME_None, RF_Transactional);
+}
+
+UMovieSceneSection* UTemplateSequenceTrack::AddNewTemplateSequenceSection(FFrameNumber KeyTime, UTemplateSequence* InSequence)
+{
+	UTemplateSequenceSection* NewSection = Cast<UTemplateSequenceSection>(CreateNewSection());
+	{
+		UMovieScene* OuterMovieScene = GetTypedOuter<UMovieScene>();
+		UMovieScene* InnerMovieScene = InSequence->GetMovieScene();
+
+		int32      InnerSequenceLength = MovieScene::DiscreteSize(InnerMovieScene->GetPlaybackRange());
+		FFrameTime OuterSequenceLength = ConvertFrameTime(InnerSequenceLength, InnerMovieScene->GetTickResolution(), OuterMovieScene->GetTickResolution());
+
+		NewSection->InitialPlacement(Sections, KeyTime, OuterSequenceLength.FrameNumber.Value, SupportsMultipleRows());
+		NewSection->SetSequence(InSequence);
+	}
+
+	AddSection(*NewSection);
+
+	return NewSection;
+}
+
+void UTemplateSequenceTrack::PostCompile(FMovieSceneEvaluationTrack& OutTrack, const FMovieSceneTrackCompilerArgs& Args) const
+{
+	// Make sure out evaluation template runs before the spawn tracks because it will have to setup the overrides.
+	OutTrack.SetEvaluationGroup(IMovieSceneTracksModule::GetEvaluationGroupName(EBuiltInEvaluationGroup::SpawnObjects));
+	OutTrack.SetEvaluationPriority(GetEvaluationPriority());
+}
+
+#if WITH_EDITORONLY_DATA
+
+FText UTemplateSequenceTrack::GetDisplayName() const
+{
+	return LOCTEXT("TrackName", "Template Animation");
+}
+
+#endif
+
+#undef LOCTEXT_NAMESPACE
diff --git a/Engine/Plugins/MovieScene/TemplateSequence/Source/TemplateSequence/Public/Evaluation/TemplateSequenceInstanceData.h b/Engine/Plugins/MovieScene/TemplateSequence/Source/TemplateSequence/Public/Evaluation/TemplateSequenceInstanceData.h
new file mode 100644
--- /dev/null
+++ b/Engine/Plugins/MovieScene/TemplateSequence/Source/TemplateSequence/Public/Evaluation/TemplateSequenceInstanceData.h
@@ -0,0 +1,20 @@
+// Copyright 1998-2019 Epic Games, Inc. All Rights Reserved.
+
+#pragma once
+
+#include "Evaluation/MovieSceneSequenceInstanceData.h"
+#include "Evaluation/MovieSceneEvaluationOperand.h"
+#include "TemplateSequenceInstanceData.generated.h"
+
+USTRUCT()
+struct FTemplateSequenceInstanceData : public FMovieSceneSequenceInstanceData
+{
+	GENERATED_BODY()
+
+	UPROPERTY()
+	FMovieSceneEvaluationOperand Operand;
+
+private:
+
+	virtual UScriptStruct& GetScriptStructImpl() const { return *StaticStruct(); }
+};
diff --git a/Engine/Plugins/MovieScene/TemplateSequence/Source/TemplateSequence/Public/Evaluation/TemplateSequenceSectionTemplate.h b/Engine/Plugins/MovieScene/TemplateSequence/Source/TemplateSequence/Public/Evaluation/TemplateSequenceSectionTemplate.h
new file mode 100644
--- /dev/null
+++ b/Engine/Plugins/MovieScene/TemplateSequence/Source/TemplateSequence/Public/Evaluation/TemplateSequenceSectionTemplate.h
@@ -0,0 +1,34 @@
+// Copyright 1998-2019 Epic Games, Inc. All Rights Reserved.
+
+#pragma once
+
+#include "CoreMinimal.h"
+#include "UObject/ObjectMacros.h"
+#include "Evaluation/MovieSceneEvalTemplate.h"
+#include "TemplateSequenceSectionTemplate.generated.h"
+
+class UTemplateSequenceSection;
+
+USTRUCT()
+struct FTemplateSequenceSectionTemplate : public FMovieSceneEvalTemplate
+{
+	GENERATED_BODY()
+
+	FTemplateSequenceSectionTemplate();
+	FTemplateSequenceSectionTemplate(const UTemplateSequenceSection& Section);
+
+	// FMovieSceneEvalTemplate interface.
+	virtual UScriptStruct& GetScriptStructImpl() const override { return *StaticStruct(); }
+	virtual void SetupOverrides() override;
+	virtual void Evaluate(const FMovieSceneEvaluationOperand& Operand, const FMovieSceneContext& Context, const FPersistentEvaluationData& PersistentData, FMovieSceneExecutionTokens& ExecutionTokens) const override;
+	virtual void Initialize(const FMovieSceneEvaluationOperand& Operand, const FMovieSceneContext& Context, FPersistentEvaluationData& PersistentData, IMovieScenePlayer& Player) const override;
+	virtual void TearDown(FPersistentEvaluationData& PersistentData, IMovieScenePlayer& Player) const override;
+
+private:
+
+	UPROPERTY()
+	FFrameNumber SectionStartTime;
+
+	UPROPERTY()
+	FMovieSceneEvaluationOperand InnerOperand;
+};
diff --git a/Engine/Plugins/MovieScene/TemplateSequence/Source/TemplateSequence/Public/Sections/TemplateSequenceSection.h b/Engine/Plugins/MovieScene/TemplateSequence/Source/TemplateSequence/Public/Sections/TemplateSequenceSection.h
new file mode 100644
--- /dev/null
+++ b/Engine/Plugins/MovieScene/TemplateSequence/Source/TemplateSequence/Public/Sections/TemplateSequenceSection.h
@@ -0,0 +1,25 @@
+// Copyright 1998-2019 Epic Games, Inc. All Rights Reserved.
+
+#pragma once
+
+#include "CoreMinimal.h"
+#include "UObject/ObjectMacros.h"
+#include "Sections/MovieSceneSubSection.h"
+#include "TemplateSequenceSection.generated.h"
+
+UCLASS(MinimalAPI)
+class UTemplateSequenceSection : public UMovieSceneSubSection
+{
+public:
+
+	GENERATED_BODY()
+
+	UTemplateSequenceSection();
+
+	// UMovieSceneSubSection interface
+	virtual FMovieSceneSubSequenceData GenerateSubSequenceData(const FSubSequenceInstanceDataParams& Params) const override;
+	virtual FMovieSceneEvalTemplatePtr GenerateTemplate() const override;
+
+	// UMovieSceneSection interface
+	virtual void OnDilated(float DilationFactor, FFrameNumber Origin) override;
+};
diff --git a/Engine/Plugins/MovieScene/TemplateSequence/Source/TemplateSequence/Public/TemplateSequence.h b/Engine/Plugins/MovieScene/TemplateSequence/Source/TemplateSequence/Public/TemplateSequence.h
new file mode 100644
--- /dev/null
+++ b/Engine/Plugins/MovieScene/TemplateSequence/Source/TemplateSequence/Public/TemplateSequence.h
@@ -0,0 +1,52 @@
+// Copyright 1998-2019 Epic Games, Inc. All Rights Reserved.
+
+#pragma once
+
+#include "MovieSceneSequence.h"
+#include "MovieScene.h"
+#include "UObject/SoftObjectPtr.h"
+#include "TemplateSequence.generated.h"
+
+/*
+ * Movie scene animation that can be instanced multiple times inside a level sequence.
+ */
+UCLASS(BlueprintType)
+class TEMPLATESEQUENCE_API UTemplateSequence : public UMovieSceneSequence
+{
+public:
+	GENERATED_BODY()
+
+	UTemplateSequence(const FObjectInitializer& ObjectInitializer);
+
+	void Initialize();
+
+	//~ UMovieSceneSequence interface
+	virtual void BindPossessableObject(const FGuid& ObjectId, UObject& PossessedObject, UObject* Context) override;
+	virtual bool CanPossessObject(UObject& Object, UObject* InPlaybackContext) const override;
+	virtual void LocateBoundObjects(const FGuid& ObjectId, UObject* Context, TArray<UObject*, TInlineAllocator<1>>& OutObjects) const override;
+	virtual UMovieScene* GetMovieScene() const override;
+	virtual UObject* GetParentObject(UObject* Object) const override;
+	virtual void UnbindPossessableObjects(const FGuid& ObjectId) override;
+	virtual void UnbindObjects(const FGuid& ObjectId, const TArray<UObject*>& InObjects, UObject* Context) override;
+	virtual void UnbindInvalidObjects(const FGuid& ObjectId, UObject* Context) override;
+
+	virtual bool AllowsSpawnableObjects() const override;
+
+#if WITH_EDITOR
+	virtual FText GetDisplayName() const override;
+#endif
+
+public:
+
+	UPROPERTY()
+	UMovieScene* MovieScene;
+
+	UPROPERTY()
+	TSoftClassPtr<AActor> BoundActorClass;
+
+	UPROPERTY()
+	TSoftObjectPtr<AActor> BoundPreviewActor;
+
+	UPROPERTY()
+	TMap<FGuid, FName> BoundActorComponents;
+};
diff --git a/Engine/Plugins/MovieScene/TemplateSequence/Source/TemplateSequence/Public/TemplateSequenceActor.h b/Engine/Plugins/MovieScene/TemplateSequence/Source/TemplateSequence/Public/TemplateSequenceActor.h
new file mode 100644
--- /dev/null
+++ b/Engine/Plugins/MovieScene/TemplateSequence/Source/TemplateSequence/Public/TemplateSequenceActor.h
@@ -0,0 +1,105 @@
+// Copyright 1998-2019 Epic Games, Inc. All Rights Reserved.
+
+#pragma once
+
+#include "CoreMinimal.h"
+#include "TemplateSequencePlayer.h"
+#include "IMovieScenePlaybackClient.h"
+#include "MovieSceneBindingOwnerInterface.h"
+#include "GameFramework/Actor.h"
+#include "UObject/ObjectMacros.h"
+#include "UObject/Object.h"
+#include "UObject/SoftObjectPath.h"
+#include "TemplateSequenceActor.generated.h"
+
+class UTemplateSequence;
+
+USTRUCT(BlueprintType)
+struct FTemplateSequenceBindingOverrideData
+{
+	GENERATED_BODY()
+
+	FTemplateSequenceBindingOverrideData()
+		: bOverridesDefault(true)
+	{
+	}
+
+	UPROPERTY(EditAnywhere, Category = "Binding")
+	TWeakObjectPtr<UObject> Object;
+
+	UPROPERTY(EditAnywhere, Category = "Binding")
+	bool bOverridesDefault;
+};
+
+UCLASS(hideCategories = (Rendering, Physics, LOD, Activation, Input))
+class TEMPLATESEQUENCE_API ATemplateSequenceActor
+	: public AActor
+	, public IMovieScenePlaybackClient
+{
+public:
+
+	GENERATED_BODY()
+
+	ATemplateSequenceActor(const FObjectInitializer& ObjectInitializer);
+
+public:
+
+	UPROPERTY(EditAnywhere, BlueprintReadOnly, Category = "Playback", meta = (ShowOnlyInnerProperties))
+	FMovieSceneSequencePlaybackSettings PlaybackSettings;
+
+	UPROPERTY(Instanced, Transient, Replicated, BlueprintReadOnly, BlueprintGetter = GetSequencePlayer, Category = "Playback", meta = (ExposeFunctionCategories = "Game|Cinematic"))
+	UTemplateSequencePlayer* SequencePlayer;
+
+	UPROPERTY(EditAnywhere, BlueprintReadOnly, Category = "General", meta = (AllowedClasses = "TemplateSequence"))
+	FSoftObjectPath TemplateSequence;
+
+	UPROPERTY(BlueprintReadOnly, Category = "General")
+	FTemplateSequenceBindingOverrideData BindingOverride;
+
+public:
+
+	UFUNCTION(BlueprintCallable, Category = "Game|Cinematic")
+	UTemplateSequence* GetSequence() const;
+
+	UFUNCTION(BlueprintCallable, Category = "Game|Cinematic")
+	UTemplateSequence* LoadSequence() const;
+
+	UFUNCTION(BlueprintCallable, Category = "Game|Cinematic")
+	void SetSequence(UTemplateSequence* InSequence);
+
+	UFUNCTION(BlueprintGetter)
+	UTemplateSequencePlayer* GetSequencePlayer() const;
+
+	UFUNCTION(BlueprintCallable, Category = "Game|Cinematic|Bindings")
+	void SetBinding(AActor* Actor);
+
+protected:
+
+	//~ Begin IMovieScenePlaybackClient interface
+	virtual bool RetrieveBindingOverrides(const FGuid& InBindingId, FMovieSceneSequenceID InSequenceID, TArray<UObject*, TInlineAllocator<1>>& OutObjects) const override;
+	virtual UObject* GetInstanceData() const override;
+	//~ End IMovieScenePlaybackClient interface
+
+	//~ Begin AActor interface
+	virtual void Tick(float DeltaSeconds) override;
+	virtual void PostInitializeComponents() override;
+	virtual void BeginPlay() override;
+	virtual void EndPlay(const EEndPlayReason::Type EndPlayReason) override;
+	//~ End AActor interface
+
+	//~ Begin UObject interface
+	virtual void PostInitProperties() override;
+	//~ End UObject interface
+
+public:
+
+	void InitializePlayer();
+
+#if WITH_EDITOR
+	virtual bool GetReferencedContentObjects(TArray<UObject*>& Objects) const override;
+#endif
+
+private:
+
+	void OnSequenceLoaded(const FName& PackageName, UPackage* Package, EAsyncLoadingResult::Type Result);
+};
diff --git a/Engine/Plugins/MovieScene/TemplateSequence/Source/TemplateSequence/Public/TemplateSequencePlayer.h b/Engine/Plugins/MovieScene/TemplateSequence/Source/TemplateSequence/Public/TemplateSequencePlayer.h
new file mode 100644
--- /dev/null
+++ b/Engine/Plugins/MovieScene/TemplateSequence/Source/TemplateSequence/Public/TemplateSequencePlayer.h
@@ -0,0 +1,36 @@
+// Copyright 1998-2019 Epic Games, Inc. All Rights Reserved.
+
+#pragma once
+
+#include "MovieSceneSequencePlayer.h"
+#include "TemplateSequencePlayer.generated.h"
+
+class ATemplateSequenceActor;
+class UTemplateSequence;
+
+UCLASS(BlueprintType)
+class TEMPLATESEQUENCE_API UTemplateSequencePlayer : public UMovieSceneSequencePlayer
+{
+public:
+
+	GENERATED_BODY()
+
+	UTemplateSequencePlayer(const FObjectInitializer&);
+
+public:
+
+	UFUNCTION(BlueprintCallable, Category = "Game|Cinematic", meta = (WorldContext = "WorldContextObject", DynamicOutputParam = "OutActor"))
+	static UTemplateSequencePlayer* CreateTemplateSequencePlayer(UObject* WorldContextObject, UTemplateSequence* TemplateSequence, FMovieSceneSequencePlaybackSettings Settings, ATemplateSequenceActor*& OutActor);
+
+public:
+
+	void Initialize(UMovieSceneSequence* InSequence, UWorld* InWorld, const FMovieSceneSequencePlaybackSettings& InSettings);
+
+	// IMovieScenePlayer interface
+	virtual UObject* GetPlaybackContext() const override;
+
+private:
+
+	/** The world this player will spawn actors in, if needed. */
+	TWeakObjectPtr<UWorld> World;
+};
diff --git a/Engine/Plugins/MovieScene/TemplateSequence/Source/TemplateSequence/Public/TemplateSequenceSpawnRegister.h b/Engine/Plugins/MovieScene/TemplateSequence/Source/TemplateSequence/Public/TemplateSequenceSpawnRegister.h
new file mode 100644
--- /dev/null
+++ b/Engine/Plugins/MovieScene/TemplateSequence/Source/TemplateSequence/Public/TemplateSequenceSpawnRegister.h
@@ -0,0 +1,13 @@
+// Copyright 1998-2019 Epic Games, Inc. All Rights Reserved.
+
+#pragma once
+
+#include "LevelSequenceSpawnRegister.h"
+
+/** Movie scene spawn register that knows how to handle spawning objects (actors) for a template sequence  */
+class TEMPLATESEQUENCE_API FTemplateSequenceSpawnRegister : public FLevelSequenceSpawnRegister
+{
+public:
+
+	FTemplateSequenceSpawnRegister();
+};
diff --git a/Engine/Plugins/MovieScene/TemplateSequence/Source/TemplateSequence/Public/Tracks/TemplateSequenceTrack.h b/Engine/Plugins/MovieScene/TemplateSequence/Source/TemplateSequence/Public/Tracks/TemplateSequenceTrack.h
new file mode 100644
--- /dev/null
+++ b/Engine/Plugins/MovieScene/TemplateSequence/Source/TemplateSequence/Public/Tracks/TemplateSequenceTrack.h
@@ -0,0 +1,36 @@
+// Copyright 1998-2019 Epic Games, Inc. All Rights Reserved.
+
+#pragma once
+
+#include "CoreMinimal.h"
+#include "UObject/ObjectMacros.h"
+#include "Misc/InlineValue.h"
+#include "Compilation/MovieSceneSegmentCompiler.h"
+#include "Tracks/MovieSceneSubTrack.h"
+#include "TemplateSequenceTrack.generated.h"
+
+class UTemplateSequence;
+class UTemplateSequenceSection;
+struct FMovieSceneEvaluationTrack;
+
+UCLASS(MinimalAPI)
+class UTemplateSequenceTrack : public UMovieSceneSubTrack
+{
+public:
+
+	GENERATED_BODY()
+
+	virtual UMovieSceneSection* AddNewTemplateSequenceSection(FFrameNumber KeyTime, UTemplateSequence* InSequence);
+
+	// UMovieSceneTrack interface
+	virtual bool SupportsType(TSubclassOf<UMovieSceneSection> SectionClass) const override;
+	virtual UMovieSceneSection* CreateNewSection() override;
+	virtual void PostCompile(FMovieSceneEvaluationTrack& OutTrack, const FMovieSceneTrackCompilerArgs& Args) const override;
+
+#if WITH_EDITORONLY_DATA
+	virtual FText GetDisplayName() const override;
+#endif
+
+private:
+	static uint16 GetEvaluationPriority() { return uint16(0x1000); }  // One more than spawn tracks.
+};
diff --git a/Engine/Plugins/MovieScene/TemplateSequence/Source/TemplateSequence/TemplateSequence.Build.cs b/Engine/Plugins/MovieScene/TemplateSequence/Source/TemplateSequence/TemplateSequence.Build.cs
new file mode 100644
--- /dev/null
+++ b/Engine/Plugins/MovieScene/TemplateSequence/Source/TemplateSequence/TemplateSequence.Build.cs
@@ -0,0 +1,28 @@
+// Copyright 1998-2019 Epic Games, Inc. All Rights Reserved.
+
+using UnrealBuildTool;
+
+public class TemplateSequence : ModuleRules
+{
+	public TemplateSequence(ReadOnlyTargetRules Target) : base(Target)
+	{
+		PrivateIncludePaths.Add("TemplateSequence/Private");
+
+		PublicDependencyModuleNames.AddRange(
+			new string[] {
+				"Core",
+				"CoreUObject",
+				"Engine",
+				"MovieScene",
+				"TimeManagement",
+				"LevelSequence"
+				}
+		);
+
+		PrivateDependencyModuleNames.AddRange(
+			new string[] {
+				"MovieSceneTracks"
+				}
+		);
+	}
+}
diff --git a/Engine/Plugins/MovieScene/TemplateSequence/Source/TemplateSequenceEditor/Private/AssetTools/TemplateSequenceActions.cpp b/Engine/Plugins/MovieScene/TemplateSequence/Source/TemplateSequenceEditor/Private/AssetTools/TemplateSequenceActions.cpp
new file mode 100644
--- /dev/null
+++ b/Engine/Plugins/MovieScene/TemplateSequence/Source/TemplateSequenceEditor/Private/AssetTools/TemplateSequenceActions.cpp
@@ -0,0 +1,76 @@
+// Copyright 1998-2019 Epic Games, Inc. All Rights Reserved.
+
+#include "AssetTools/TemplateSequenceActions.h"
+#include "EngineGlobals.h"
+#include "Engine/Engine.h"
+#include "TemplateSequence.h"
+#include "TemplateSequenceEditorToolkit.h"
+
+#define LOCTEXT_NAMESPACE "AssetTypeActions"
+
+FTemplateSequenceActions::FTemplateSequenceActions(const TSharedRef<ISlateStyle>& InStyle)
+	: Style(InStyle)
+{ }
+
+uint32 FTemplateSequenceActions::GetCategories()
+{
+	return EAssetTypeCategories::Animation;
+}
+
+FText FTemplateSequenceActions::GetName() const
+{
+	return NSLOCTEXT("AssetTypeActions", "AssetTypeActions_TemplateSequence", "Template Sequence");
+}
+
+UClass* FTemplateSequenceActions::GetSupportedClass() const
+{
+	return UTemplateSequence::StaticClass();
+}
+
+FColor FTemplateSequenceActions::GetTypeColor() const
+{
+	return FColor(200, 80, 80);
+}
+
+void FTemplateSequenceActions::OpenAssetEditor(const TArray<UObject*>& InObjects, TSharedPtr<IToolkitHost> EditWithinLevelEditor)
+{
+	UWorld* WorldContext = nullptr;
+	for (const FWorldContext& Context : GEngine->GetWorldContexts())
+	{
+		if (Context.WorldType == EWorldType::Editor)
+		{
+			WorldContext = Context.World();
+			break;
+		}
+	}
+
+	if (!ensure(WorldContext))
+	{
+		return;
+	}
+
+	EToolkitMode::Type Mode = EditWithinLevelEditor.IsValid()
+		? EToolkitMode::WorldCentric
+		: EToolkitMode::Standalone;
+
+	for (auto ObjIt = InObjects.CreateConstIterator(); ObjIt; ++ObjIt)
+	{
+		UTemplateSequence* TemplateSequence = Cast<UTemplateSequence>(*ObjIt);
+
+		if (TemplateSequence != nullptr)
+		{
+			TSharedRef<FTemplateSequenceEditorToolkit> Toolkit = MakeShareable(new FTemplateSequenceEditorToolkit(Style));
+			Toolkit->Initialize(Mode, EditWithinLevelEditor, TemplateSequence);
+		}
+	}
+}
+
+
+bool FTemplateSequenceActions::ShouldForceWorldCentric()
+{
+	// @todo sequencer: Hack to force world-centric mode for Sequencer
+	return true;
+}
+
+
+#undef LOCTEXT_NAMESPACE
diff --git a/Engine/Plugins/MovieScene/TemplateSequence/Source/TemplateSequenceEditor/Private/AssetTools/TemplateSequenceActions.h b/Engine/Plugins/MovieScene/TemplateSequence/Source/TemplateSequenceEditor/Private/AssetTools/TemplateSequenceActions.h
new file mode 100644
--- /dev/null
+++ b/Engine/Plugins/MovieScene/TemplateSequence/Source/TemplateSequenceEditor/Private/AssetTools/TemplateSequenceActions.h
@@ -0,0 +1,41 @@
+// Copyright 1998-2019 Epic Games, Inc. All Rights Reserved.
+
+#pragma once
+
+#include "CoreMinimal.h"
+#include "Styling/ISlateStyle.h"
+#include "Toolkits/IToolkitHost.h"
+#include "AssetTypeActions_Base.h"
+
+
+/**
+ * Implements actions for UTemplateSequence assets.
+ */
+class FTemplateSequenceActions : public FAssetTypeActions_Base
+{
+public:
+
+	/**
+	 * Creates and initializes a new instance.
+	 *
+	 * @param InStyle The style set to use for asset editor toolkits.
+	 */
+	FTemplateSequenceActions(const TSharedRef<ISlateStyle>& InStyle);
+
+public:
+
+	// IAssetTypeActions interface
+	virtual uint32 GetCategories() override;
+	virtual FText GetName() const override;
+	virtual UClass* GetSupportedClass() const override;
+	virtual FColor GetTypeColor() const override;
+	virtual void OpenAssetEditor(const TArray<UObject*>& InObjects, TSharedPtr<class IToolkitHost> EditWithinLevelEditor = TSharedPtr<IToolkitHost>()) override;
+	virtual bool ShouldForceWorldCentric() override;
+	virtual bool CanLocalize() const override { return false; }
+	virtual bool HasActions(const TArray<UObject*>& InObjects) const override { return false; }
+
+private:
+
+	/** Pointer to the style set to use for toolkits. */
+	TSharedRef<ISlateStyle> Style;
+};
diff --git a/Engine/Plugins/MovieScene/TemplateSequence/Source/TemplateSequenceEditor/Private/Factories/TemplateSequenceFactoryNew.cpp b/Engine/Plugins/MovieScene/TemplateSequence/Source/TemplateSequenceEditor/Private/Factories/TemplateSequenceFactoryNew.cpp
new file mode 100644
--- /dev/null
+++ b/Engine/Plugins/MovieScene/TemplateSequence/Source/TemplateSequenceEditor/Private/Factories/TemplateSequenceFactoryNew.cpp
@@ -0,0 +1,37 @@
+// Copyright 1998-2019 Epic Games, Inc. All Rights Reserved.
+
+#include "Factories/TemplateSequenceFactoryNew.h"
+#include "TemplateSequence.h"
+#include "MovieScene.h"
+#include "MovieSceneToolsProjectSettings.h"
+
+#define LOCTEXT_NAMESPACE "MovieSceneFactory"
+
+UTemplateSequenceFactoryNew::UTemplateSequenceFactoryNew(const FObjectInitializer& ObjectInitializer)
+	: Super(ObjectInitializer)
+{
+	bCreateNew = true;
+	bEditAfterNew = true;
+	SupportedClass = UTemplateSequence::StaticClass();
+}
+
+UObject* UTemplateSequenceFactoryNew::FactoryCreateNew(UClass* Class, UObject* InParent, FName Name, EObjectFlags Flags, UObject* Context, FFeedbackContext* Warn)
+{
+	auto NewTemplateSequence = NewObject<UTemplateSequence>(InParent, Name, Flags | RF_Transactional);
+	NewTemplateSequence->Initialize();
+
+	// Set up some sensible defaults
+	const UMovieSceneToolsProjectSettings* ProjectSettings = GetDefault<UMovieSceneToolsProjectSettings>();
+
+	FFrameRate TickResolution = NewTemplateSequence->GetMovieScene()->GetTickResolution();
+	NewTemplateSequence->GetMovieScene()->SetPlaybackRange((ProjectSettings->DefaultStartTime*TickResolution).FloorToFrame(), (ProjectSettings->DefaultDuration*TickResolution).FloorToFrame().Value);
+
+	return NewTemplateSequence;
+}
+
+bool UTemplateSequenceFactoryNew::ShouldShowInNewMenu() const
+{
+	return true;
+}
+
+#undef LOCTEXT_NAMESPACE
diff --git a/Engine/Plugins/MovieScene/TemplateSequence/Source/TemplateSequenceEditor/Private/Factories/TemplateSequenceFactoryNew.h b/Engine/Plugins/MovieScene/TemplateSequence/Source/TemplateSequenceEditor/Private/Factories/TemplateSequenceFactoryNew.h
new file mode 100644
--- /dev/null
+++ b/Engine/Plugins/MovieScene/TemplateSequence/Source/TemplateSequenceEditor/Private/Factories/TemplateSequenceFactoryNew.h
@@ -0,0 +1,22 @@
+// Copyright 1998-2019 Epic Games, Inc. All Rights Reserved.
+
+#pragma once
+
+#include "CoreMinimal.h"
+#include "UObject/ObjectMacros.h"
+#include "Factories/Factory.h"
+#include "TemplateSequenceFactoryNew.generated.h"
+
+/**
+ * Implements a factory for UTemplateSequence objects.
+ */
+UCLASS(hidecategories = Object)
+class UTemplateSequenceFactoryNew : public UFactory
+{
+	GENERATED_UCLASS_BODY()
+
+public:
+	// UFactory Interface
+	virtual UObject* FactoryCreateNew(UClass* Class, UObject* InParent, FName Name, EObjectFlags Flags, UObject* Context, FFeedbackContext* Warn) override;
+	virtual bool ShouldShowInNewMenu() const override;
+};
diff --git a/Engine/Plugins/MovieScene/TemplateSequence/Source/TemplateSequenceEditor/Private/Misc/MovieSceneSequenceEditor_TemplateSequence.h b/Engine/Plugins/MovieScene/TemplateSequence/Source/TemplateSequenceEditor/Private/Misc/MovieSceneSequenceEditor_TemplateSequence.h
new file mode 100644
--- /dev/null
+++ b/Engine/Plugins/MovieScene/TemplateSequence/Source/TemplateSequenceEditor/Private/Misc/MovieSceneSequenceEditor_TemplateSequence.h
@@ -0,0 +1,9 @@
+// Copyright 1998-2019 Epic Games, Inc. All Rights Reserved.
+
+#pragma once
+
+#include "MovieSceneSequenceEditor.h"
+
+struct FMovieSceneSequenceEditor_TemplateSequence : FMovieSceneSequenceEditor
+{
+};
diff --git a/Engine/Plugins/MovieScene/TemplateSequence/Source/TemplateSequenceEditor/Private/Misc/TemplateSequenceEditorPlaybackContext.cpp b/Engine/Plugins/MovieScene/TemplateSequence/Source/TemplateSequenceEditor/Private/Misc/TemplateSequenceEditorPlaybackContext.cpp
new file mode 100644
--- /dev/null
+++ b/Engine/Plugins/MovieScene/TemplateSequence/Source/TemplateSequenceEditor/Private/Misc/TemplateSequenceEditorPlaybackContext.cpp
@@ -0,0 +1,49 @@
+// Copyright 1998-2019 Epic Games, Inc. All Rights Reserved.
+
+#include "Misc/TemplateSequenceEditorPlaybackContext.h"
+#include "Editor.h"
+
+#define LOCTEXT_NAMESPACE "TemplateSequenceEditorPlaybackContext"
+
+UObject* FTemplateSequenceEditorPlaybackContext::GetPlaybackContext() const
+{
+	UWorld* Context = WeakCurrentContext.Get();
+	if (Context)
+	{
+		return Context;
+	}
+
+	Context = ComputePlaybackContext();
+	check(Context);
+	WeakCurrentContext = Context;
+	return Context;
+}
+
+UWorld* FTemplateSequenceEditorPlaybackContext::ComputePlaybackContext()
+{
+	const bool bIsSimulatingInEditor = GEditor && GEditor->bIsSimulatingInEditor;
+
+	UWorld* EditorWorld = nullptr;
+
+	for (const FWorldContext& Context : GEngine->GetWorldContexts())
+	{
+		if (Context.WorldType == EWorldType::PIE)
+		{
+			UWorld* ThisWorld = Context.World();
+			if (ThisWorld)
+			{
+				EditorWorld = ThisWorld;
+				break;
+			}
+		}
+		else if (Context.WorldType == EWorldType::Editor)
+		{
+			EditorWorld = Context.World();
+		}
+	}
+
+	check(EditorWorld);
+	return EditorWorld;
+}
+
+#undef LOCTEXT_NAMESPACE
diff --git a/Engine/Plugins/MovieScene/TemplateSequence/Source/TemplateSequenceEditor/Private/Misc/TemplateSequenceEditorPlaybackContext.h b/Engine/Plugins/MovieScene/TemplateSequence/Source/TemplateSequenceEditor/Private/Misc/TemplateSequenceEditorPlaybackContext.h
new file mode 100644
--- /dev/null
+++ b/Engine/Plugins/MovieScene/TemplateSequence/Source/TemplateSequenceEditor/Private/Misc/TemplateSequenceEditorPlaybackContext.h
@@ -0,0 +1,25 @@
+// Copyright 1998-2019 Epic Games, Inc. All Rights Reserved.
+
+#pragma once
+
+#include "Templates/SharedPointer.h"
+#include "UObject/WeakObjectPtr.h"
+
+class UWorld;
+
+/**
+ * Class that manages the current UWorld context that a level-sequence editor should use for playback
+ */
+class FTemplateSequenceEditorPlaybackContext : public TSharedFromThis<FTemplateSequenceEditorPlaybackContext>
+{
+public:
+
+	UObject* GetPlaybackContext() const;
+
+private:
+
+	static UWorld* ComputePlaybackContext();
+	
+	/** Mutable cached context pointer */
+	mutable TWeakObjectPtr<UWorld> WeakCurrentContext;
+};
diff --git a/Engine/Plugins/MovieScene/TemplateSequence/Source/TemplateSequenceEditor/Private/Misc/TemplateSequenceEditorSettings.cpp b/Engine/Plugins/MovieScene/TemplateSequence/Source/TemplateSequenceEditor/Private/Misc/TemplateSequenceEditorSettings.cpp
new file mode 100644
--- /dev/null
+++ b/Engine/Plugins/MovieScene/TemplateSequence/Source/TemplateSequenceEditor/Private/Misc/TemplateSequenceEditorSettings.cpp
@@ -0,0 +1,9 @@
+// Copyright 1998-2019 Epic Games, Inc. All Rights Reserved.
+
+#include "Misc/TemplateSequenceEditorSettings.h"
+
+UTemplateSequenceEditorSettings::UTemplateSequenceEditorSettings(const FObjectInitializer& ObjectInitializer)
+	: Super(ObjectInitializer)
+{
+}
+
diff --git a/Engine/Plugins/MovieScene/TemplateSequence/Source/TemplateSequenceEditor/Private/Misc/TemplateSequenceEditorSettings.h b/Engine/Plugins/MovieScene/TemplateSequence/Source/TemplateSequenceEditor/Private/Misc/TemplateSequenceEditorSettings.h
new file mode 100644
--- /dev/null
+++ b/Engine/Plugins/MovieScene/TemplateSequence/Source/TemplateSequenceEditor/Private/Misc/TemplateSequenceEditorSettings.h
@@ -0,0 +1,26 @@
+// Copyright 1998-2019 Epic Games, Inc. All Rights Reserved.
+
+#pragma once
+
+#include "CoreMinimal.h"
+#include "UObject/ObjectMacros.h"
+#include "UObject/Object.h"
+#include "Engine/EngineTypes.h"
+#include "UObject/SoftObjectPath.h"
+#include "TemplateSequence.h"
+#include "TemplateSequenceEditorSettings.generated.h"
+
+/**
+ * Template Sequence Editor settings.
+ */
+UCLASS(config=BaseTemplateSequenceSettings)
+class UTemplateSequenceEditorSettings
+	: public UObject
+{
+	GENERATED_BODY()
+
+public:
+
+	UTemplateSequenceEditorSettings(const FObjectInitializer& ObjectInitializer);
+};
+
diff --git a/Engine/Plugins/MovieScene/TemplateSequence/Source/TemplateSequenceEditor/Private/Misc/TemplateSequenceEditorSpawnRegister.cpp b/Engine/Plugins/MovieScene/TemplateSequence/Source/TemplateSequenceEditor/Private/Misc/TemplateSequenceEditorSpawnRegister.cpp
new file mode 100644
--- /dev/null
+++ b/Engine/Plugins/MovieScene/TemplateSequence/Source/TemplateSequenceEditor/Private/Misc/TemplateSequenceEditorSpawnRegister.cpp
@@ -0,0 +1,3 @@
+// Copyright 1998-2019 Epic Games, Inc. All Rights Reserved.
+
+#include "TemplateSequenceEditorSpawnRegister.h"
diff --git a/Engine/Plugins/MovieScene/TemplateSequence/Source/TemplateSequenceEditor/Private/Misc/TemplateSequenceEditorSpawnRegister.h b/Engine/Plugins/MovieScene/TemplateSequence/Source/TemplateSequenceEditor/Private/Misc/TemplateSequenceEditorSpawnRegister.h
new file mode 100644
--- /dev/null
+++ b/Engine/Plugins/MovieScene/TemplateSequence/Source/TemplateSequenceEditor/Private/Misc/TemplateSequenceEditorSpawnRegister.h
@@ -0,0 +1,10 @@
+// Copyright 1998-2019 Epic Games, Inc. All Rights Reserved.
+
+#pragma once
+
+#include "Misc/LevelSequenceEditorSpawnRegister.h"
+
+/** Movie scene spawn register that knows how to handle spawning objects (actors) for a template sequence  */
+class FTemplateSequenceEditorSpawnRegister : public FLevelSequenceEditorSpawnRegister
+{
+};
diff --git a/Engine/Plugins/MovieScene/TemplateSequence/Source/TemplateSequenceEditor/Private/Styles/TemplateSequenceEditorStyle.cpp b/Engine/Plugins/MovieScene/TemplateSequence/Source/TemplateSequenceEditor/Private/Styles/TemplateSequenceEditorStyle.cpp
new file mode 100644
--- /dev/null
+++ b/Engine/Plugins/MovieScene/TemplateSequence/Source/TemplateSequenceEditor/Private/Styles/TemplateSequenceEditorStyle.cpp
@@ -0,0 +1,5 @@
+// Copyright 1998-2019 Epic Games, Inc. All Rights Reserved.
+
+#include "Styles/TemplateSequenceEditorStyle.h"
+
+TSharedPtr<FTemplateSequenceEditorStyle> FTemplateSequenceEditorStyle::Singleton;
diff --git a/Engine/Plugins/MovieScene/TemplateSequence/Source/TemplateSequenceEditor/Private/Styles/TemplateSequenceEditorStyle.h b/Engine/Plugins/MovieScene/TemplateSequence/Source/TemplateSequenceEditor/Private/Styles/TemplateSequenceEditorStyle.h
new file mode 100644
--- /dev/null
+++ b/Engine/Plugins/MovieScene/TemplateSequence/Source/TemplateSequenceEditor/Private/Styles/TemplateSequenceEditorStyle.h
@@ -0,0 +1,61 @@
+// Copyright 1998-2019 Epic Games, Inc. All Rights Reserved.
+
+#pragma once
+
+#include "CoreMinimal.h"
+#include "Layout/Margin.h"
+#include "Fonts/SlateFontInfo.h"
+#include "Misc/Paths.h"
+#include "Brushes/SlateBorderBrush.h"
+#include "Brushes/SlateBoxBrush.h"
+#include "Brushes/SlateImageBrush.h"
+#include "Styling/SlateStyleRegistry.h"
+#include "Styling/SlateStyle.h"
+
+#define IMAGE_BRUSH(RelativePath, ...) FSlateImageBrush(RootToContentDir(RelativePath, TEXT(".png")), __VA_ARGS__)
+
+/**
+ * Implements the visual style of the template sequence editor.
+ */
+class FTemplateSequenceEditorStyle : public FSlateStyleSet
+{
+public:
+	
+	FTemplateSequenceEditorStyle()
+		: FSlateStyleSet("TemplateSequenceEditorStyle")
+	{
+		const FVector2D Icon16x16(16.0f, 16.0f);
+		const FVector2D Icon64x64(64.0f, 64.0f);
+
+		SetContentRoot(FPaths::EnginePluginsDir() / TEXT("MovieScene/TemplateSequence/Content"));
+
+		// tab icons
+		Set("TemplateSequenceEditor.Tabs.Sequencer", new IMAGE_BRUSH("icon_tab_sequencer_16x", Icon16x16));
+
+		// asset thumbnail
+		Set("ClassThumbnail.TemplateSequence", new IMAGE_BRUSH("TemplateSequence_16x", Icon16x16));
+		Set("ClassThumbnail.TemplateSequence", new IMAGE_BRUSH("TemplateSequence_64x", Icon64x64));
+
+		FSlateStyleRegistry::RegisterSlateStyle(*this);
+	}
+
+	virtual ~FTemplateSequenceEditorStyle()
+	{
+		FSlateStyleRegistry::UnRegisterSlateStyle(*this);
+	}
+
+	static TSharedRef<FTemplateSequenceEditorStyle> Get()
+	{
+		if (!Singleton.IsValid())
+		{
+			Singleton = MakeShareable(new FTemplateSequenceEditorStyle);
+		}
+		return Singleton.ToSharedRef();
+	}
+
+private:
+
+	static TSharedPtr<FTemplateSequenceEditorStyle> Singleton;
+};
+
+#undef IMAGE_BRUSH
diff --git a/Engine/Plugins/MovieScene/TemplateSequence/Source/TemplateSequenceEditor/Private/TemplateSequenceEditorModule.cpp b/Engine/Plugins/MovieScene/TemplateSequence/Source/TemplateSequenceEditor/Private/TemplateSequenceEditorModule.cpp
new file mode 100644
--- /dev/null
+++ b/Engine/Plugins/MovieScene/TemplateSequence/Source/TemplateSequenceEditor/Private/TemplateSequenceEditorModule.cpp
@@ -0,0 +1,149 @@
+// Copyright 1998-2019 Epic Games, Inc. All Rights Reserved.
+
+#include "AssetTools/TemplateSequenceActions.h"
+#include "AssetToolsModule.h"
+#include "ISequencerModule.h"
+#include "ISettingsModule.h"
+#include "Misc/MovieSceneSequenceEditor_TemplateSequence.h"
+#include "Misc/TemplateSequenceEditorSettings.h"
+#include "Modules/ModuleManager.h"
+#include "SequencerSettings.h"
+#include "Styles/TemplateSequenceEditorStyle.h"
+#include "TemplateSequence.h"
+#include "TrackEditors/TemplateSequenceTrackEditor.h"
+#include "UObject/Class.h"
+#include "UObject/GCObject.h"
+
+#define LOCTEXT_NAMESPACE "TemplateSequenceEditor"
+
+/**
+ * Implements the FTemplateSequenceEditor module.
+ */
+class FTemplateSequenceEditorModule : public IModuleInterface, public FGCObject
+{
+public:
+	FTemplateSequenceEditorModule()
+		: Settings(nullptr)
+	{
+	}
+
+	virtual void StartupModule() override
+	{
+		RegisterSettings();
+		RegisterAssetTools();
+		RegisterSequenceEditor();
+		RegisterTrackEditors();
+	}
+
+	virtual void ShutdownModule() override
+	{
+		UnregisterTrackEditors();
+		UnregisterSequenceEditor();
+		UnregisterAssetTools();
+		UnregisterSettings();
+	}
+
+	/** FGCObject interface */
+	virtual void AddReferencedObjects(FReferenceCollector& Collector) override
+	{
+		if (Settings != nullptr)
+		{
+			Collector.AddReferencedObject(Settings);
+		}
+	}
+
+private:
+
+	void RegisterSettings()
+	{
+		ISettingsModule* SettingsModule = FModuleManager::GetModulePtr<ISettingsModule>("Settings");
+
+		if (SettingsModule != nullptr)
+		{
+			SettingsModule->RegisterSettings("Project", "Plugins", "TemplateSequencer",
+				LOCTEXT("TemplateSequencerSettingsName", "Template Sequencer"),
+				LOCTEXT("TemplateSequencerSettingsDescription", "Configure the Template Sequence Editor."),
+				GetMutableDefault<UTemplateSequenceEditorSettings>()
+			);
+
+			Settings = USequencerSettingsContainer::GetOrCreate<USequencerSettings>(TEXT("TemplateSequenceEditor"));
+
+			SettingsModule->RegisterSettings("Editor", "ContentEditors", "TemplateSequenceEditor",
+				LOCTEXT("TemplateSequenceEditorSettingsName", "Template Sequence Editor"),
+				LOCTEXT("TemplateSequenceEditorSettingsDescription", "Configure the look and feel of the Template Sequence Editor."),
+				Settings);
+		}
+	}
+
+	void UnregisterSettings()
+	{
+		ISettingsModule* SettingsModule = FModuleManager::GetModulePtr<ISettingsModule>("Settings");
+
+		if (SettingsModule != nullptr)
+		{
+			SettingsModule->UnregisterSettings("Project", "Plugins", "TemplateSequencer");
+			SettingsModule->UnregisterSettings("Editor", "ContentEditors", "TemplateSequenceEditor");
+		}
+	}
+
+	void RegisterAssetTools()
+	{
+		IAssetTools& AssetTools = FModuleManager::LoadModuleChecked<FAssetToolsModule>("AssetTools").Get();
+		TemplateSequenceTypeActions = MakeShared<FTemplateSequenceActions>(FTemplateSequenceEditorStyle::Get());
+		AssetTools.RegisterAssetTypeActions(TemplateSequenceTypeActions.ToSharedRef());
+	}
+
+	void UnregisterAssetTools()
+	{
+		FAssetToolsModule* AssetToolsModule = FModuleManager::GetModulePtr<FAssetToolsModule>("AssetTools");
+
+		if (AssetToolsModule != nullptr)
+		{
+			IAssetTools& AssetTools = AssetToolsModule->Get();
+			AssetTools.UnregisterAssetTypeActions(TemplateSequenceTypeActions.ToSharedRef());
+		}
+	}
+
+	void RegisterSequenceEditor()
+	{
+		ISequencerModule& SequencerModule = FModuleManager::LoadModuleChecked<ISequencerModule>("Sequencer");
+		SequenceEditorHandle = SequencerModule.RegisterSequenceEditor(UTemplateSequence::StaticClass(), MakeUnique<FMovieSceneSequenceEditor_TemplateSequence>());
+	}
+
+	void UnregisterSequenceEditor()
+	{
+		ISequencerModule* SequencerModule = FModuleManager::GetModulePtr<ISequencerModule>("Sequencer");
+		if (SequencerModule)
+		{
+			SequencerModule->UnregisterSequenceEditor(SequenceEditorHandle);
+		}
+	}
+
+	void RegisterTrackEditors()
+	{
+		ISequencerModule& SequencerModule = FModuleManager::LoadModuleChecked<ISequencerModule>("Sequencer");
+		TemplateSequenceTrackCreateEditorHandle = SequencerModule.RegisterTrackEditor(FOnCreateTrackEditor::CreateStatic(&FTemplateSequenceTrackEditor::CreateTrackEditor));
+	}
+
+	void UnregisterTrackEditors()
+	{
+		ISequencerModule* SequencerModule = FModuleManager::GetModulePtr<ISequencerModule>("Sequencer");
+		if (SequencerModule)
+		{
+			SequencerModule->UnRegisterTrackEditor(TemplateSequenceTrackCreateEditorHandle);
+		}
+	}
+
+private:
+
+	TSharedPtr<FTemplateSequenceActions> TemplateSequenceTypeActions;
+
+	FDelegateHandle SequenceEditorHandle;
+	FDelegateHandle TemplateSequenceTrackCreateEditorHandle;
+
+	USequencerSettings* Settings;
+};
+
+IMPLEMENT_MODULE(FTemplateSequenceEditorModule, TemplateSequenceEditor);
+
+#undef LOCTEXT_NAMESPACE
diff --git a/Engine/Plugins/MovieScene/TemplateSequence/Source/TemplateSequenceEditor/Private/TemplateSequenceEditorToolkit.cpp b/Engine/Plugins/MovieScene/TemplateSequence/Source/TemplateSequenceEditor/Private/TemplateSequenceEditorToolkit.cpp
new file mode 100644
--- /dev/null
+++ b/Engine/Plugins/MovieScene/TemplateSequence/Source/TemplateSequenceEditor/Private/TemplateSequenceEditorToolkit.cpp
@@ -0,0 +1,439 @@
+// Copyright 1998-2019 Epic Games, Inc. All Rights Reserved.
+
+#include "TemplateSequenceEditorToolkit.h"
+#include "ClassViewerFilter.h"
+#include "ClassViewerModule.h"
+#include "DragAndDrop/ActorDragDropGraphEdOp.h"
+#include "DragAndDrop/AssetDragDropOp.h"
+#include "DragAndDrop/ClassDragDropOp.h"
+#include "Engine/Selection.h"
+#include "Framework/Docking/TabManager.h"
+#include "Framework/MultiBox/MultiBoxBuilder.h"
+#include "ISequencer.h"
+#include "ISequencerModule.h"
+#include "LevelEditor.h"
+#include "LevelEditorSequencerIntegration.h"
+#include "Misc/TemplateSequenceEditorPlaybackContext.h"
+#include "Misc/TemplateSequenceEditorSpawnRegister.h"
+#include "Modules/ModuleManager.h"
+#include "ScopedTransaction.h"
+#include "SequencerSettings.h"
+#include "Widgets/Docking/SDockTab.h"
+#include "Widgets/Input/SButton.h"
+#include "Widgets/Text/STextBlock.h"
+
+#define LOCTEXT_NAMESPACE "TemplateSequenceEditor"
+
+const FName FTemplateSequenceEditorToolkit::SequencerMainTabId(TEXT("Sequencer_SequencerMain"));
+
+namespace SequencerDefs
+{
+	static const FName SequencerAppIdentifier(TEXT("SequencerApp"));
+}
+
+FTemplateSequenceEditorToolkit::FTemplateSequenceEditorToolkit(const TSharedRef<ISlateStyle>& InStyle)
+	: TemplateSequence(nullptr)
+	, Style(InStyle)
+{
+	ISequencerModule& SequencerModule = FModuleManager::Get().LoadModuleChecked<ISequencerModule>("Sequencer");
+	int32 NewIndex = SequencerModule.GetAddTrackMenuExtensibilityManager()->GetExtenderDelegates().Add(
+		FAssetEditorExtender::CreateRaw(this, &FTemplateSequenceEditorToolkit::HandleMenuExtensibilityGetExtender));
+	SequencerExtenderHandle = SequencerModule.GetAddTrackMenuExtensibilityManager()->GetExtenderDelegates()[NewIndex].GetHandle();
+}
+
+FTemplateSequenceEditorToolkit::~FTemplateSequenceEditorToolkit()
+{
+	FLevelEditorSequencerIntegration::Get().RemoveSequencer(Sequencer.ToSharedRef());
+
+	Sequencer->Close();
+
+	if (FModuleManager::Get().IsModuleLoaded(TEXT("LevelEditor")))
+	{
+		auto& LevelEditorModule = FModuleManager::LoadModuleChecked<FLevelEditorModule>(TEXT("LevelEditor"));
+		LevelEditorModule.OnMapChanged().RemoveAll(this);
+	}
+
+	ISequencerModule& SequencerModule = FModuleManager::Get().LoadModuleChecked<ISequencerModule>("Sequencer");
+	SequencerModule.GetAddTrackMenuExtensibilityManager()->GetExtenderDelegates().RemoveAll([this](const FAssetEditorExtender& Extender)
+	{
+		return SequencerExtenderHandle == Extender.GetHandle();
+	});
+}
+
+void FTemplateSequenceEditorToolkit::Initialize(const EToolkitMode::Type Mode, const TSharedPtr<IToolkitHost>& InitToolkitHost, UTemplateSequence* InTemplateSequence)
+{
+	// create tab layout
+	const TSharedRef<FTabManager::FLayout> StandaloneDefaultLayout = FTabManager::NewLayout("Standalone_TemplateSequenceEditor")
+		->AddArea
+		(
+			FTabManager::NewPrimaryArea()
+			->Split
+			(
+				FTabManager::NewStack()
+				->AddTab(SequencerMainTabId, ETabState::OpenedTab)
+			)
+		);
+
+	TemplateSequence = InTemplateSequence;
+	PlaybackContext = MakeShared<FTemplateSequenceEditorPlaybackContext>();
+
+	const bool bCreateDefaultStandaloneMenu = true;
+	const bool bCreateDefaultToolbar = false;
+
+	FAssetEditorToolkit::InitAssetEditor(Mode, InitToolkitHost, SequencerDefs::SequencerAppIdentifier, StandaloneDefaultLayout, bCreateDefaultStandaloneMenu, bCreateDefaultToolbar, TemplateSequence);
+
+	TSharedRef<FTemplateSequenceEditorSpawnRegister> SpawnRegister = MakeShareable(new FTemplateSequenceEditorSpawnRegister());
+	SpawnRegister->SetSequencer(Sequencer);
+
+	// Initialize sequencer.
+	FSequencerInitParams SequencerInitParams;
+	{
+		SequencerInitParams.RootSequence = TemplateSequence;
+		SequencerInitParams.bEditWithinLevelEditor = true;
+		SequencerInitParams.ToolkitHost = InitToolkitHost;
+		SequencerInitParams.SpawnRegister = SpawnRegister;
+		SequencerInitParams.HostCapabilities.bSupportsCurveEditor = true;
+		SequencerInitParams.HostCapabilities.bSupportsDiscardChanges = true;
+		SequencerInitParams.HostCapabilities.bSupportsSaveMovieSceneAsset = true;
+
+		SequencerInitParams.PlaybackContext.Bind(PlaybackContext.ToSharedRef(), &FTemplateSequenceEditorPlaybackContext::GetPlaybackContext);
+
+		SequencerInitParams.ViewParams.UniqueName = "TemplateSequenceEditor";
+		SequencerInitParams.ViewParams.ScrubberStyle = ESequencerScrubberStyle::FrameBlock;
+		SequencerInitParams.ViewParams.OnReceivedFocus.BindRaw(this, &FTemplateSequenceEditorToolkit::OnSequencerReceivedFocus);
+		SequencerInitParams.ViewParams.OnAssetsDrop.BindRaw(this, &FTemplateSequenceEditorToolkit::OnSequencerAssetsDrop);
+		SequencerInitParams.ViewParams.OnClassesDrop.BindRaw(this, &FTemplateSequenceEditorToolkit::OnSequencerClassesDrop);
+		SequencerInitParams.ViewParams.OnActorsDrop.BindRaw(this, &FTemplateSequenceEditorToolkit::OnSequencerActorsDrop);
+
+		TSharedRef<FExtender> ToolbarExtender = MakeShared<FExtender>();
+		ToolbarExtender->AddToolBarExtension("Base Commands", EExtensionHook::After, nullptr, FToolBarExtensionDelegate::CreateSP(this, &FTemplateSequenceEditorToolkit::ExtendSequencerToolbar));
+		SequencerInitParams.ViewParams.ToolbarExtender = ToolbarExtender;
+	}
+
+	Sequencer = FModuleManager::LoadModuleChecked<ISequencerModule>("Sequencer").CreateSequencer(SequencerInitParams);
+
+	Sequencer->OnActorAddedToSequencer().AddSP(this, &FTemplateSequenceEditorToolkit::HandleActorAddedToSequencer);
+
+	FLevelEditorSequencerIntegrationOptions Options;
+	Options.bRequiresLevelEvents = true;
+	Options.bRequiresActorEvents = true;
+	Options.bCanRecord = true;
+	FLevelEditorSequencerIntegration::Get().AddSequencer(Sequencer.ToSharedRef(), Options);
+
+	FLevelEditorModule& LevelEditorModule = FModuleManager::LoadModuleChecked<FLevelEditorModule>("LevelEditor");
+
+	// Reopen the scene outliner so that is refreshed with the sequencer info column
+	if (Sequencer->GetSequencerSettings()->GetShowOutlinerInfoColumn())
+	{
+		TSharedPtr<FTabManager> LevelEditorTabManager = LevelEditorModule.GetLevelEditorTabManager();
+		if (LevelEditorTabManager->FindExistingLiveTab(FName("LevelEditorSceneOutliner")).IsValid())
+		{
+			LevelEditorTabManager->InvokeTab(FName("LevelEditorSceneOutliner"))->RequestCloseTab();
+			LevelEditorTabManager->InvokeTab(FName("LevelEditorSceneOutliner"));
+		}
+	}
+
+	LevelEditorModule.AttachSequencer(Sequencer->GetSequencerWidget(), SharedThis(this));
+	LevelEditorModule.OnMapChanged().AddRaw(this, &FTemplateSequenceEditorToolkit::HandleMapChanged);
+}
+
+FText FTemplateSequenceEditorToolkit::GetBaseToolkitName() const
+{
+	return LOCTEXT("AppLabel", "Template Sequence Editor");
+}
+
+FName FTemplateSequenceEditorToolkit::GetToolkitFName() const
+{
+	static FName SequencerName("TemplateSequenceEditor");
+	return SequencerName;
+}
+
+FString FTemplateSequenceEditorToolkit::GetWorldCentricTabPrefix() const
+{
+	return LOCTEXT("WorldCentricTabPrefix", "Sequencer ").ToString();
+}
+
+FLinearColor FTemplateSequenceEditorToolkit::GetWorldCentricTabColorScale() const
+{
+	return FLinearColor(0.7, 0.0f, 0.0f, 0.5f);
+}
+
+void FTemplateSequenceEditorToolkit::RegisterTabSpawners(const TSharedRef<class FTabManager>& InTabManager)
+{
+	if (IsWorldCentricAssetEditor())
+	{
+		return;
+	}
+}
+
+
+void FTemplateSequenceEditorToolkit::UnregisterTabSpawners(const TSharedRef<class FTabManager>& InTabManager)
+{
+	if (!IsWorldCentricAssetEditor())
+	{
+		InTabManager->UnregisterTabSpawner(SequencerMainTabId);
+	}
+
+	FLevelEditorModule& LevelEditorModule = FModuleManager::LoadModuleChecked<FLevelEditorModule>("LevelEditor");
+	LevelEditorModule.AttachSequencer(SNullWidget::NullWidget, nullptr);
+}
+
+void FTemplateSequenceEditorToolkit::ExtendSequencerToolbar(FToolBarBuilder& ToolbarBuilder)
+{
+	TSharedRef<SHorizontalBox> Widget = SNew(SHorizontalBox)
+		+SHorizontalBox::Slot()
+		.AutoWidth()
+		.VAlign(VAlign_Center)
+		[
+			SNew(STextBlock)
+			.Text(LOCTEXT("BoundActorClassPicker", "Bound Actor Class"))
+		]
+		+SHorizontalBox::Slot()
+		.AutoWidth()
+		.VAlign(VAlign_Center)
+		[
+			SNew(SComboButton)
+			.OnGetMenuContent(this, &FTemplateSequenceEditorToolkit::GetBoundActorClassMenuContent)
+			.ButtonContent()
+			[
+				SNew(STextBlock)
+				.Text(this, &FTemplateSequenceEditorToolkit::GetBoundActorClassName)
+			]
+		];
+	
+	ToolbarBuilder.AddWidget(Widget);
+}
+
+FText FTemplateSequenceEditorToolkit::GetBoundActorClassName() const
+{
+	const UClass* BoundActorClass = TemplateSequence ? TemplateSequence->BoundActorClass.Get() : NULL;
+	return BoundActorClass ? BoundActorClass->GetDisplayNameText() : FText::FromName(NAME_None);
+}
+
+TSharedRef<SWidget> FTemplateSequenceEditorToolkit::GetBoundActorClassMenuContent() const
+{
+	FClassViewerModule& ClassViewerModule = FModuleManager::LoadModuleChecked<FClassViewerModule>("ClassViewer");
+
+	FClassViewerInitializationOptions Options;
+	Options.Mode = EClassViewerMode::ClassPicker;
+	Options.bIsActorsOnly = true;
+
+	TSharedRef<SWidget> ClassPicker = ClassViewerModule.CreateClassViewer(Options, FOnClassPicked::CreateSP(this, &FTemplateSequenceEditorToolkit::OnBoundActorClassPicked));
+
+	return SNew(SBox)
+		.WidthOverride(350.0f)
+		[
+			SNew(SVerticalBox)
+			+ SVerticalBox::Slot()
+			.MaxHeight(400.0f)
+			.AutoHeight()
+			[
+				ClassPicker
+			]
+		];
+}
+
+void FTemplateSequenceEditorToolkit::OnBoundActorClassPicked(UClass* ChosenClass)
+{
+	FSlateApplication::Get().DismissAllMenus();
+
+	if (TemplateSequence != nullptr)
+	{
+		ChangeActorBinding(*ChosenClass);
+
+		Sequencer->NotifyMovieSceneDataChanged(EMovieSceneDataChangeType::MovieSceneStructureItemAdded);
+	}
+}
+
+TSharedRef<FExtender> FTemplateSequenceEditorToolkit::HandleMenuExtensibilityGetExtender(const TSharedRef<FUICommandList> CommandList, const TArray<UObject*> ContextSensitiveObjects)
+{
+	TSharedRef<FExtender> AddTrackMenuExtender(new FExtender());
+	AddTrackMenuExtender->AddMenuExtension(
+		SequencerMenuExtensionPoints::AddTrackMenu_PropertiesSection,
+		EExtensionHook::Before,
+		CommandList,
+		FMenuExtensionDelegate::CreateRaw(this, &FTemplateSequenceEditorToolkit::HandleTrackMenuExtensionAddTrack, ContextSensitiveObjects));
+
+	return AddTrackMenuExtender;
+}
+
+void FTemplateSequenceEditorToolkit::HandleTrackMenuExtensionAddTrack(FMenuBuilder& AddTrackMenuBuilder, TArray<UObject*> ContextObjects)
+{
+	// TODO-lchabant: stolen from level sequence.
+	if (ContextObjects.Num() != 1)
+	{
+		return;
+	}
+
+	AActor* Actor = Cast<AActor>(ContextObjects[0]);
+	if (Actor == nullptr)
+	{
+		return;
+	}
+
+	AddTrackMenuBuilder.BeginSection("Components", LOCTEXT("ComponentsSection", "Components"));
+	{
+		for (UActorComponent* Component : Actor->GetComponents())
+		{
+			if (Component)
+			{
+				FUIAction AddComponentAction(FExecuteAction::CreateSP(this, &FTemplateSequenceEditorToolkit::HandleAddComponentActionExecute, Component));
+				FText AddComponentLabel = FText::FromString(Component->GetName());
+				FText AddComponentToolTip = FText::Format(LOCTEXT("ComponentToolTipFormat", "Add {0} component"), FText::FromString(Component->GetName()));
+				AddTrackMenuBuilder.AddMenuEntry(AddComponentLabel, AddComponentToolTip, FSlateIcon(), AddComponentAction);
+			}
+		}
+	}
+	AddTrackMenuBuilder.EndSection();
+}
+
+void FTemplateSequenceEditorToolkit::HandleAddComponentActionExecute(UActorComponent* Component)
+{
+	// TODO-lchabant: stolen from level sequence.
+	const FScopedTransaction Transaction(LOCTEXT("AddComponent", "Add Component"));
+
+	FString ComponentName = Component->GetName();
+
+	TArray<UActorComponent*> ActorComponents;
+	ActorComponents.Add(Component);
+
+	USelection* SelectedActors = GEditor->GetSelectedActors();
+	if (SelectedActors && SelectedActors->Num() > 0)
+	{
+		for (FSelectionIterator Iter(*SelectedActors); Iter; ++Iter)
+		{
+			AActor* Actor = CastChecked<AActor>(*Iter);
+
+			TArray<UActorComponent*> OutActorComponents;
+			Actor->GetComponents(OutActorComponents);
+
+			for (UActorComponent* ActorComponent : OutActorComponents)
+			{
+				if (ActorComponent->GetName() == ComponentName)
+				{
+					ActorComponents.AddUnique(ActorComponent);
+				}
+			}
+		}
+	}
+
+	for (UActorComponent* ActorComponent : ActorComponents)
+	{
+		Sequencer->GetHandleToObject(ActorComponent);
+	}
+}
+
+void FTemplateSequenceEditorToolkit::HandleActorAddedToSequencer(AActor* Actor, const FGuid Binding)
+{
+	// TODO-lchabant: add default tracks (re-use level sequence toolkit code).
+}
+
+void FTemplateSequenceEditorToolkit::HandleMapChanged(UWorld* NewWorld, EMapChangeType MapChangeType)
+{
+	if ((MapChangeType == EMapChangeType::LoadMap || MapChangeType == EMapChangeType::NewMap || MapChangeType == EMapChangeType::TearDownWorld))
+	{
+		Sequencer->GetSpawnRegister().CleanUp(*Sequencer);
+		CloseWindow();
+	}
+}
+
+bool FTemplateSequenceEditorToolkit::OnRequestClose()
+{
+	return true;
+}
+
+bool FTemplateSequenceEditorToolkit::CanFindInContentBrowser() const
+{
+	// False so that sequencer doesn't take over Find In Content Browser functionality and always find the level sequence asset.
+	return false;
+}
+
+void FTemplateSequenceEditorToolkit::OnSequencerReceivedFocus()
+{
+	if (Sequencer.IsValid())
+	{
+		FLevelEditorSequencerIntegration::Get().OnSequencerReceivedFocus(Sequencer.ToSharedRef());
+	}
+}
+
+bool FTemplateSequenceEditorToolkit::OnSequencerReceivedDragOver(const FGeometry& MyGeometry, const FDragDropEvent& DragDropEvent, FReply& OutReply)
+{
+	bool bIsDragSupported = false;
+
+	TSharedPtr<FDragDropOperation> Operation = DragDropEvent.GetOperation();
+	if (Operation.IsValid() && (
+		(Operation->IsOfType<FAssetDragDropOp>() && StaticCastSharedPtr<FAssetDragDropOp>(Operation)->GetAssetPaths().Num() <= 1) ||
+		(Operation->IsOfType<FClassDragDropOp>() && StaticCastSharedPtr<FClassDragDropOp>(Operation)->ClassesToDrop.Num() <= 1) ||
+		(Operation->IsOfType<FActorDragDropGraphEdOp>() && StaticCastSharedPtr<FActorDragDropGraphEdOp>(Operation)->Actors.Num() <= 1)))
+	{
+		bIsDragSupported = true;
+	}
+
+	OutReply = (bIsDragSupported ? FReply::Handled() : FReply::Unhandled());
+	return true;
+}
+
+bool FTemplateSequenceEditorToolkit::OnSequencerAssetsDrop(const TArray<UObject*>& Assets, const FAssetDragDropOp& DragDropOp)
+{
+	if (Assets.Num() > 0)
+	{
+		// Only drop the first asset.
+		UObject* CurObject = Assets[0];
+
+		// TODO: check for dropping a sequence?
+
+		ChangeActorBinding(*CurObject, DragDropOp.GetActorFactory());
+	}
+
+	return true;
+}
+
+bool FTemplateSequenceEditorToolkit::OnSequencerClassesDrop(const TArray<TWeakObjectPtr<UClass>>& Classes, const FClassDragDropOp& DragDropOp)
+{
+	return false;
+}
+
+bool FTemplateSequenceEditorToolkit::OnSequencerActorsDrop(const TArray<TWeakObjectPtr<AActor>>& Actors, const FActorDragDropGraphEdOp& DragDropOp)
+{
+	return false;
+}
+
+void FTemplateSequenceEditorToolkit::ChangeActorBinding(UObject& Object, UActorFactory* ActorFactory, bool bSetupDefaults)
+{
+	UMovieScene* MovieScene = TemplateSequence->GetMovieScene();
+	check(MovieScene != nullptr);
+
+	// See if we previously had a main object binding.
+	FGuid PreviousSpawnableGuid;
+	if (MovieScene->GetSpawnableCount() > 0)
+		PreviousSpawnableGuid = MovieScene->GetSpawnable(0).GetGuid();
+
+	// Make the new spawnable object binding.
+	FGuid NewSpawnableGuid = Sequencer->MakeNewSpawnable(Object, ActorFactory, bSetupDefaults);
+	FMovieSceneSpawnable* NewSpawnable = MovieScene->FindSpawnable(NewSpawnableGuid);
+
+	if (Object.IsA<UClass>())
+	{
+		UClass* ChosenClass = StaticCast<UClass*>(&Object);
+		TemplateSequence->BoundActorClass = ChosenClass;
+	}
+	else
+	{
+		const UObject* SpawnableTemplate = NewSpawnable->GetObjectTemplate();
+		TemplateSequence->BoundActorClass = SpawnableTemplate->GetClass();
+	}
+
+	// If we had a previous one, move everything under it to the new binding, and clean up.
+	if (PreviousSpawnableGuid.IsValid())
+	{
+		MovieScene->MoveBindingContents(PreviousSpawnableGuid, NewSpawnableGuid);
+
+		if (MovieScene->RemoveSpawnable(PreviousSpawnableGuid))
+		{
+			FMovieSceneSpawnRegister& SpawnRegister = Sequencer->GetSpawnRegister();
+			SpawnRegister.DestroySpawnedObject(PreviousSpawnableGuid, Sequencer->GetFocusedTemplateID(), *Sequencer);
+		}
+	}
+}
+
+#undef LOCTEXT_NAMESPACE
diff --git a/Engine/Plugins/MovieScene/TemplateSequence/Source/TemplateSequenceEditor/Private/TemplateSequenceEditorToolkit.h b/Engine/Plugins/MovieScene/TemplateSequence/Source/TemplateSequenceEditor/Private/TemplateSequenceEditorToolkit.h
new file mode 100644
--- /dev/null
+++ b/Engine/Plugins/MovieScene/TemplateSequence/Source/TemplateSequenceEditor/Private/TemplateSequenceEditorToolkit.h
@@ -0,0 +1,109 @@
+// Copyright 1998-2019 Epic Games, Inc. All Rights Reserved.
+
+#pragma once
+
+#include "CoreMinimal.h"
+#include "TemplateSequence.h"
+#include "UObject/GCObject.h"
+#include "Styling/ISlateStyle.h"
+#include "Toolkits/AssetEditorToolkit.h"
+
+class FToolBarBuilder;
+class FTemplateSequenceEditorPlaybackContext;
+class ISequencer;
+class FAssetDragDropOp;
+class FClassDragDropOp;
+class FActorDragDropGraphEdOp;
+
+/**
+ * Implements an Editor toolkit for template sequences.
+ */
+class FTemplateSequenceEditorToolkit : public FAssetEditorToolkit, public FGCObject
+{
+public:
+
+	/**
+	 * Creates and initializes a new instance.
+	 *
+	 * @param InStyle The style set to use.
+	 */
+	FTemplateSequenceEditorToolkit(const TSharedRef<ISlateStyle>& InStyle);
+
+	/** Virtual destructor */
+	virtual ~FTemplateSequenceEditorToolkit();
+
+public:
+
+	/**
+	 * Initialize this asset editor.
+	 *
+	 * @param Mode Asset editing mode for this editor (standalone or world-centric).
+	 * @param InitToolkitHost When Mode is WorldCentric, this is the level editor instance to spawn this editor within.
+	 * @param TemplateSequence The animation to edit.
+	 * @param TrackEditorDelegates Delegates to call to create auto-key handlers for this sequencer.
+	 */
+	void Initialize(const EToolkitMode::Type Mode, const TSharedPtr<IToolkitHost>& InitToolkitHost, UTemplateSequence* TemplateSequence);
+
+public:
+
+	//~ FAssetEditorToolkit interface
+	virtual void AddReferencedObjects(FReferenceCollector& Collector) override
+	{
+		Collector.AddReferencedObject(TemplateSequence);
+	}
+
+	virtual bool OnRequestClose() override;
+	virtual bool CanFindInContentBrowser() const override;
+
+public:
+
+	//~ IToolkit interface
+	virtual FText GetBaseToolkitName() const override;
+	virtual FName GetToolkitFName() const override;
+	virtual FString GetWorldCentricTabPrefix() const override;
+	virtual FLinearColor GetWorldCentricTabColorScale() const override;
+	virtual void RegisterTabSpawners(const TSharedRef<FTabManager>& TabManager) override;
+	virtual void UnregisterTabSpawners(const TSharedRef<FTabManager>& TabManager) override;
+
+private:
+
+	void ExtendSequencerToolbar(FToolBarBuilder& ToolbarBuilder);
+	FText GetBoundActorClassName() const;
+	TSharedRef<SWidget> GetBoundActorClassMenuContent() const;
+	void OnBoundActorClassPicked(UClass* ChosenClass);
+
+	TSharedRef<FExtender> HandleMenuExtensibilityGetExtender(const TSharedRef<FUICommandList> CommandList, const TArray<UObject*> ContextSensitiveObjects);
+	void HandleTrackMenuExtensionAddTrack(FMenuBuilder& AddTrackMenuBuilder, TArray<UObject*> ContextObjects);
+	void HandleAddComponentActionExecute(UActorComponent* Component);
+
+	void HandleActorAddedToSequencer(AActor* Actor, const FGuid Binding);
+
+	void HandleMapChanged(class UWorld* NewWorld, EMapChangeType MapChangeType);
+
+	void OnSequencerReceivedFocus();
+	bool OnSequencerReceivedDragOver(const FGeometry& MyGeometry, const FDragDropEvent& DragDropEvent, FReply& OutReply);
+	bool OnSequencerAssetsDrop(const TArray<UObject*>& Assets, const FAssetDragDropOp& DragDropOp);
+	bool OnSequencerClassesDrop(const TArray<TWeakObjectPtr<UClass>>& Classes, const FClassDragDropOp& DragDropOp);
+	bool OnSequencerActorsDrop(const TArray<TWeakObjectPtr<AActor>>& Actors, const FActorDragDropGraphEdOp& DragDropOp);
+
+	void ChangeActorBinding(UObject& Object, UActorFactory* ActorFactory = nullptr, bool bSetupDefaults = true);
+
+private:
+
+	/** Template sequence for our edit operation. */
+	UTemplateSequence* TemplateSequence;
+
+	/** The sequencer used by this editor. */
+	TSharedPtr<ISequencer> Sequencer;
+
+	/** Pointer to the style set to use for toolkits. */
+	TSharedRef<ISlateStyle> Style;
+
+	/** Handle to the sequencer properties menu extender. */
+	FDelegateHandle SequencerExtenderHandle;
+
+	TSharedPtr<FTemplateSequenceEditorPlaybackContext> PlaybackContext;
+
+	/**	The tab ids for all the tabs used */
+	static const FName SequencerMainTabId;
+};
diff --git a/Engine/Plugins/MovieScene/TemplateSequence/Source/TemplateSequenceEditor/Private/TrackEditors/TemplateSequenceTrackEditor.cpp b/Engine/Plugins/MovieScene/TemplateSequence/Source/TemplateSequenceEditor/Private/TrackEditors/TemplateSequenceTrackEditor.cpp
new file mode 100644
--- /dev/null
+++ b/Engine/Plugins/MovieScene/TemplateSequence/Source/TemplateSequenceEditor/Private/TrackEditors/TemplateSequenceTrackEditor.cpp
@@ -0,0 +1,167 @@
+// Copyright 1998-2019 Epic Games, Inc. All Rights Reserved.
+
+#include "TrackEditors/TemplateSequenceTrackEditor.h"
+#include "AssetData.h"
+#include "AssetRegistryModule.h"
+#include "CommonMovieSceneTools.h"
+#include "ContentBrowserModule.h"
+#include "Framework/Application/SlateApplication.h"
+#include "Framework/MultiBox/MultiBoxBuilder.h"
+#include "IContentBrowserSingleton.h"
+#include "MovieSceneTimeHelpers.h"
+#include "Sections/TemplateSequenceSection.h"
+#include "SequencerSectionPainter.h"
+#include "TemplateSequence.h"
+#include "Tracks/TemplateSequenceTrack.h"
+#include "Widgets/Layout/SBox.h"
+
+#define LOCTEXT_NAMESPACE "FTemplateSequenceTrackEditor"
+
+FTemplateSequenceTrackEditor::FTemplateSequenceTrackEditor(TSharedRef<ISequencer> InSequencer)
+	: FMovieSceneTrackEditor(InSequencer)
+{
+}
+
+TSharedRef<ISequencerTrackEditor> FTemplateSequenceTrackEditor::CreateTrackEditor(TSharedRef<ISequencer> InSequencer)
+{
+	return MakeShareable(new FTemplateSequenceTrackEditor(InSequencer));
+}
+
+bool FTemplateSequenceTrackEditor::SupportsType(TSubclassOf<UMovieSceneTrack> Type) const
+{
+	return Type == UTemplateSequenceTrack::StaticClass();
+}
+
+void FTemplateSequenceTrackEditor::BuildObjectBindingTrackMenu(FMenuBuilder& MenuBuilder, const TArray<FGuid>& ObjectBindings, const UClass* ObjectClass)
+{
+	MenuBuilder.AddSubMenu(
+		LOCTEXT("AddTemplateSequence", "Template Sequence"), NSLOCTEXT("Sequencer", "AddTemplateSequenceTooltip", "Adds a track that can play a template sequence asset using the parent binding."),
+		FNewMenuDelegate::CreateRaw(this, &FTemplateSequenceTrackEditor::AddTemplateSequenceSubMenu, ObjectBindings)
+	);
+}
+
+TSharedRef<ISequencerSection> FTemplateSequenceTrackEditor::MakeSectionInterface(UMovieSceneSection& SectionObject, UMovieSceneTrack& Track, FGuid ObjectBinding)
+{
+	check( SupportsType( SectionObject.GetOuter()->GetClass() ) );
+	
+	return MakeShareable(new FTemplateSequenceSection(*CastChecked<UTemplateSequenceSection>(&SectionObject)));
+}
+
+void FTemplateSequenceTrackEditor::AddTemplateSequenceSubMenu(FMenuBuilder& MenuBuilder, TArray<FGuid> ObjectBindings)
+{
+	FAssetPickerConfig AssetPickerConfig;
+	{
+		AssetPickerConfig.OnAssetSelected = FOnAssetSelected::CreateRaw(this, &FTemplateSequenceTrackEditor::OnTemplateSequenceAssetSelected, ObjectBindings);
+		AssetPickerConfig.OnAssetEnterPressed = FOnAssetEnterPressed::CreateRaw(this, &FTemplateSequenceTrackEditor::OnTemplateSequenceAssetEnterPressed, ObjectBindings);
+		AssetPickerConfig.bAllowNullSelection = false;
+		AssetPickerConfig.InitialAssetViewType = EAssetViewType::List;
+		AssetPickerConfig.Filter.ClassNames.Add(UTemplateSequence::StaticClass()->GetFName());
+	}
+
+	FContentBrowserModule& ContentBrowserModule = FModuleManager::Get().LoadModuleChecked<FContentBrowserModule>(TEXT("ContentBrowser"));
+
+	TSharedPtr<SBox> MenuEntry = SNew(SBox)
+		.WidthOverride(300.0f)
+		.HeightOverride(300.f)
+		[
+			ContentBrowserModule.Get().CreateAssetPicker(AssetPickerConfig)
+		];
+
+	MenuBuilder.AddWidget(MenuEntry.ToSharedRef(), FText::GetEmpty(), true);
+}
+
+void FTemplateSequenceTrackEditor::OnTemplateSequenceAssetSelected(const FAssetData& AssetData, TArray<FGuid> ObjectBindings)
+{
+	FSlateApplication::Get().DismissAllMenus();
+
+	UTemplateSequence* SelectedSequence = Cast<UTemplateSequence>(AssetData.GetAsset());
+	if (SelectedSequence != nullptr)
+	{
+		const FScopedTransaction Transaction(LOCTEXT("AddTemplateSequence_Transaction", "Add Template Animation"));
+
+		AnimatablePropertyChanged(FOnKeyProperty::CreateRaw(this, &FTemplateSequenceTrackEditor::AddKeyInternal, ObjectBindings, SelectedSequence));
+	}
+}
+
+void FTemplateSequenceTrackEditor::OnTemplateSequenceAssetEnterPressed(const TArray<FAssetData>& AssetData, TArray<FGuid> ObjectBindings)
+{
+	if (AssetData.Num() > 0)
+	{
+		OnTemplateSequenceAssetSelected(AssetData[0].GetAsset(), ObjectBindings);
+	}
+}
+
+FKeyPropertyResult FTemplateSequenceTrackEditor::AddKeyInternal(FFrameNumber KeyTime, TArray<FGuid> ObjectBindings, UTemplateSequence* TemplateSequence)
+{
+	FKeyPropertyResult KeyPropertyResult;
+
+	bool bHandleCreated = false;
+	bool bTrackCreated = false;
+	bool bTrackModified = false;
+
+	TSharedPtr<ISequencer> SequencerPtr = GetSequencer();
+	if (SequencerPtr.IsValid())
+	{
+		for (const FGuid& ObjectBindingGuid : ObjectBindings)
+		{
+			UObject* Object = SequencerPtr->FindSpawnedObjectOrTemplate(ObjectBindingGuid);
+
+			FFindOrCreateHandleResult HandleResult = FindOrCreateHandleToObject(Object);
+			FGuid ObjectHandle = HandleResult.Handle;
+			KeyPropertyResult.bHandleCreated |= HandleResult.bWasCreated;
+
+			if (ObjectHandle.IsValid())
+			{
+				FFindOrCreateTrackResult TrackResult = FindOrCreateTrackForObject(ObjectHandle, UTemplateSequenceTrack::StaticClass());
+				UMovieSceneTrack* Track = TrackResult.Track;
+				KeyPropertyResult.bTrackCreated |= TrackResult.bWasCreated;
+
+				if (ensure(Track))
+				{
+					UMovieSceneSection* NewSection = Cast<UTemplateSequenceTrack>(Track)->AddNewTemplateSequenceSection(KeyTime, TemplateSequence);
+					KeyPropertyResult.bTrackModified = true;
+					
+					GetSequencer()->EmptySelection();
+					GetSequencer()->SelectSection(NewSection);
+					GetSequencer()->ThrobSectionSelection();
+				}
+			}
+		}
+	}
+
+	return KeyPropertyResult;
+}
+
+FTemplateSequenceSection::FTemplateSequenceSection(UTemplateSequenceSection& InSection)
+	: WeakSection(&InSection)
+{
+}
+
+UMovieSceneSection* FTemplateSequenceSection::GetSectionObject()
+{
+	return WeakSection.Get();
+}
+
+bool FTemplateSequenceSection::IsReadOnly() const
+{
+	return WeakSection.IsValid() ? WeakSection.Get()->IsReadOnly() : false;
+}
+
+FText FTemplateSequenceSection::GetSectionTitle() const
+{
+	if (const UTemplateSequenceSection* Section = WeakSection.Get())
+	{
+		if (const UMovieSceneSequence* TemplateSequence = Section->GetSequence())
+		{
+			return FText::FromString(TemplateSequence->GetName());
+		}
+	}
+	return LOCTEXT("NoTemplateSequenceSection", "No Template Animation");
+}
+
+int32 FTemplateSequenceSection::OnPaintSection(FSequencerSectionPainter& Painter) const
+{
+	return Painter.PaintSectionBackground();
+}
+
+#undef LOCTEXT_NAMESPACE
diff --git a/Engine/Plugins/MovieScene/TemplateSequence/Source/TemplateSequenceEditor/Private/TrackEditors/TemplateSequenceTrackEditor.h b/Engine/Plugins/MovieScene/TemplateSequence/Source/TemplateSequenceEditor/Private/TrackEditors/TemplateSequenceTrackEditor.h
new file mode 100644
--- /dev/null
+++ b/Engine/Plugins/MovieScene/TemplateSequence/Source/TemplateSequenceEditor/Private/TrackEditors/TemplateSequenceTrackEditor.h
@@ -0,0 +1,59 @@
+// Copyright 1998-2019 Epic Games, Inc. All Rights Reserved.
+
+#pragma once
+
+#include "CoreMinimal.h"
+#include "Misc/Guid.h"
+#include "MovieSceneTrackEditor.h"
+
+struct FAssetData;
+class FMenuBuilder;
+class UTemplateSequence;
+class UTemplateSequenceSection;
+
+class FTemplateSequenceTrackEditor : public FMovieSceneTrackEditor
+{
+public:
+
+	FTemplateSequenceTrackEditor(TSharedRef<ISequencer> InSequencer);
+
+	static TSharedRef<ISequencerTrackEditor> CreateTrackEditor(TSharedRef<ISequencer> OwningSequencer);
+
+public:
+
+	// ISequencerTrackEditor interface
+	virtual bool SupportsType(TSubclassOf<class UMovieSceneTrack> TrackClass) const override;
+	virtual void BuildObjectBindingTrackMenu(FMenuBuilder& MenuBuilder, const TArray<FGuid>& ObjectBindings, const UClass* ObjectClass) override;
+	virtual TSharedRef<ISequencerSection> MakeSectionInterface(UMovieSceneSection& SectionObject, UMovieSceneTrack& Track, FGuid ObjectBinding) override;
+
+private:
+	void AddTemplateSequenceSubMenu(FMenuBuilder& MenuBuilder, TArray<FGuid> ObjectBindings);
+	void OnTemplateSequenceAssetSelected(const FAssetData& AssetData, TArray<FGuid> ObjectBindings);
+	void OnTemplateSequenceAssetEnterPressed(const TArray<FAssetData>& AssetData, TArray<FGuid> ObjectBindings);
+
+	FKeyPropertyResult AddKeyInternal(FFrameNumber KeyTime, TArray<FGuid> ObjectBindings, UTemplateSequence* TemplateSequence);
+};
+
+class FTemplateSequenceSection
+	: public ISequencerSection
+	, public TSharedFromThis<FTemplateSequenceSection>
+{
+public:
+
+	/** Constructor. */
+	FTemplateSequenceSection(UTemplateSequenceSection& InSection);
+
+	/** Virtual destructor. */
+	virtual ~FTemplateSequenceSection() {}
+
+	// ISequencerSection interface
+	virtual UMovieSceneSection* GetSectionObject() override;
+	virtual bool IsReadOnly() const override;
+	virtual FText GetSectionTitle() const override;
+	virtual int32 OnPaintSection(FSequencerSectionPainter& Painter) const override;
+
+private:
+
+	/** The section we are visualizing */
+	TWeakObjectPtr<UTemplateSequenceSection> WeakSection;
+};
diff --git a/Engine/Plugins/MovieScene/TemplateSequence/Source/TemplateSequenceEditor/TemplateSequenceEditor.Build.cs b/Engine/Plugins/MovieScene/TemplateSequence/Source/TemplateSequenceEditor/TemplateSequenceEditor.Build.cs
new file mode 100644
--- /dev/null
+++ b/Engine/Plugins/MovieScene/TemplateSequence/Source/TemplateSequenceEditor/TemplateSequenceEditor.Build.cs
@@ -0,0 +1,48 @@
+// Copyright 1998-2019 Epic Games, Inc. All Rights Reserved.
+
+using UnrealBuildTool;
+
+public class TemplateSequenceEditor : ModuleRules
+{
+	public TemplateSequenceEditor(ReadOnlyTargetRules Target) : base(Target)
+	{
+		PrivateDependencyModuleNames.AddRange(
+			new string[] {
+				"TemplateSequence",
+				"BlueprintGraph",
+				"Core",
+				"CoreUObject",
+				"Engine",
+				"InputCore",
+				"Kismet",
+				"MovieScene",
+				"MovieSceneTools",
+				"Sequencer",
+				"EditorStyle",
+				"Slate",
+				"SlateCore",
+				"UnrealEd",
+				"TimeManagement",
+				"LevelSequence",
+				"LevelSequenceEditor"
+			}
+		);
+
+		PrivateIncludePaths.AddRange(
+			new string[] {
+				"TemplateSequenceEditor/Private",
+			}
+		);
+
+		var DynamicModuleNames = new string[] {
+			"LevelEditor",
+			"PropertyEditor",
+		};
+
+		foreach (var Name in DynamicModuleNames)
+		{
+			PrivateIncludePathModuleNames.Add(Name);
+			DynamicallyLoadedModuleNames.Add(Name);
+		}
+	}
+}
diff --git a/Engine/Plugins/MovieScene/TemplateSequence/TemplateSequence.uplugin b/Engine/Plugins/MovieScene/TemplateSequence/TemplateSequence.uplugin
new file mode 100644
--- /dev/null
+++ b/Engine/Plugins/MovieScene/TemplateSequence/TemplateSequence.uplugin
@@ -0,0 +1,37 @@
+{
+	"FileVersion" : 1,
+	"Version" : 1,
+	"VersionName" : "0.1",
+	"FriendlyName" : "Template Sequence (Experimental)",
+	"Description" : "Runtime for template sequences",
+	"Category" : "Editor",
+	"CreatedBy" : "Epic Games, Inc.",
+	"CreatedByURL" : "http://epicgames.com",
+	"DocsURL" : "",
+	"MarketplaceURL" : "",
+	"SupportURL" : "",
+	"EnabledByDefault" : false,
+	"CanContainContent" : false,
+	"IsBetaVersion" : true,
+	"Installed" : false,
+	"Plugins":
+	[
+		{
+			"Name": "LevelSequenceEditor",
+			"Enabled": true
+		}
+	],
+	"Modules" :
+	[
+		{
+			"Name" : "TemplateSequence",
+			"Type" : "Runtime",
+			"LoadingPhase" : "Default"
+		},
+		{
+			"Name" : "TemplateSequenceEditor",
+			"Type" : "Editor",
+			"LoadingPhase" : "PostEngineInit"
+		}
+	]
+}
