Fix deprecated function call
#rnx

diff --git a/Engine/Plugins/MovieScene/TemplateSequence/Source/TemplateSequenceEditor/Private/TemplateSequenceEditorToolkit.cpp b/Engine/Plugins/MovieScene/TemplateSequence/Source/TemplateSequenceEditor/Private/TemplateSequenceEditorToolkit.cpp
--- a/Engine/Plugins/MovieScene/TemplateSequence/Source/TemplateSequenceEditor/Private/TemplateSequenceEditorToolkit.cpp
+++ b/Engine/Plugins/MovieScene/TemplateSequence/Source/TemplateSequenceEditor/Private/TemplateSequenceEditorToolkit.cpp
@@ -209,7 +209,7 @@ FText FTemplateSequenceEditorToolkit::GetBoundActorClassName() const
 	return BoundActorClass ? BoundActorClass->GetDisplayNameText() : FText::FromName(NAME_None);
 }
 
-TSharedRef<SWidget> FTemplateSequenceEditorToolkit::GetBoundActorClassMenuContent() const
+TSharedRef<SWidget> FTemplateSequenceEditorToolkit::GetBoundActorClassMenuContent()
 {
 	FClassViewerModule& ClassViewerModule = FModuleManager::LoadModuleChecked<FClassViewerModule>("ClassViewer");
 
diff --git a/Engine/Plugins/MovieScene/TemplateSequence/Source/TemplateSequenceEditor/Private/TemplateSequenceEditorToolkit.h b/Engine/Plugins/MovieScene/TemplateSequence/Source/TemplateSequenceEditor/Private/TemplateSequenceEditorToolkit.h
--- a/Engine/Plugins/MovieScene/TemplateSequence/Source/TemplateSequenceEditor/Private/TemplateSequenceEditorToolkit.h
+++ b/Engine/Plugins/MovieScene/TemplateSequence/Source/TemplateSequenceEditor/Private/TemplateSequenceEditorToolkit.h
@@ -69,7 +69,7 @@ private:
 
 	void ExtendSequencerToolbar(FToolBarBuilder& ToolbarBuilder);
 	FText GetBoundActorClassName() const;
-	TSharedRef<SWidget> GetBoundActorClassMenuContent() const;
+	TSharedRef<SWidget> GetBoundActorClassMenuContent();
 	void OnBoundActorClassPicked(UClass* ChosenClass);
 
 	TSharedRef<FExtender> HandleMenuExtensibilityGetExtender(const TSharedRef<FUICommandList> CommandList, const TArray<UObject*> ContextSensitiveObjects);
