Sequencer: Fix template sequence binding overrides for when the hierarchy is deeper than one level.

diff --git a/Engine/Plugins/MovieScene/TemplateSequence/Source/TemplateSequence/Private/Evaluation/TemplateSequenceSectionTemplate.cpp b/Engine/Plugins/MovieScene/TemplateSequence/Source/TemplateSequence/Private/Evaluation/TemplateSequenceSectionTemplate.cpp
--- a/Engine/Plugins/MovieScene/TemplateSequence/Source/TemplateSequence/Private/Evaluation/TemplateSequenceSectionTemplate.cpp
+++ b/Engine/Plugins/MovieScene/TemplateSequence/Source/TemplateSequence/Private/Evaluation/TemplateSequenceSectionTemplate.cpp
@@ -3,6 +3,7 @@
 #include "Evaluation/TemplateSequenceSectionTemplate.h"
 #include "IMovieScenePlaybackClient.h"
 #include "TemplateSequence.h"
+#include "Evaluation/MovieSceneEvaluationTemplateInstance.h"
 #include "Sections/TemplateSequenceSection.h"
 
 FTemplateSequenceSectionTemplate::FTemplateSequenceSectionTemplate()
@@ -28,18 +29,24 @@ FTemplateSequenceSectionTemplate::FTemplateSequenceSectionTemplate(const UTempla
 
 void FTemplateSequenceSectionTemplate::SetupOverrides()
 {
-	EnableOverrides(RequiresInitializeFlag | RequiresTearDownFlag);
+	EnableOverrides(RequiresSetupFlag | RequiresTearDownFlag);
 }
 
 void FTemplateSequenceSectionTemplate::Evaluate(const FMovieSceneEvaluationOperand& Operand, const FMovieSceneContext& Context, const FPersistentEvaluationData& PersistentData, FMovieSceneExecutionTokens& ExecutionTokens) const
 {
 }
 
-void FTemplateSequenceSectionTemplate::Initialize(const FMovieSceneEvaluationOperand& Operand, const FMovieSceneContext& Context, FPersistentEvaluationData& PersistentData, IMovieScenePlayer& Player) const
+void FTemplateSequenceSectionTemplate::Setup(FPersistentEvaluationData& PersistentData, IMovieScenePlayer& Player) const
 {
 	if (InnerOperand.IsValid())
 	{
-		Player.BindingOverrides.Add(InnerOperand, Operand);
+		// Add the override that will tell the child template sequence to use our object binding.
+		const FMovieSceneSequenceID SequenceID = PersistentData.GetSectionKey().SequenceID;
+		const FMovieSceneObjectBindingID AbsoluteBinding = GetAbsoluteInnerBindingID(SequenceID, Player);
+		const FMovieSceneEvaluationOperand AbsoluteInnerOperand(AbsoluteBinding.GetSequenceID(), AbsoluteBinding.GetGuid());
+
+		const FMovieSceneEvaluationOperand Operand(SequenceID, OuterBindingId);
+		Player.BindingOverrides.Add(AbsoluteInnerOperand, Operand);
 	}
 }
 
@@ -47,6 +54,20 @@ void FTemplateSequenceSectionTemplate::TearDown(FPersistentEvaluationData& Persi
 {
 	if (InnerOperand.IsValid())
 	{
-		Player.BindingOverrides.Remove(InnerOperand);
+		// Clear the override.
+		const FMovieSceneSequenceID SequenceID = PersistentData.GetSectionKey().SequenceID;
+		const FMovieSceneObjectBindingID AbsoluteBinding = GetAbsoluteInnerBindingID(SequenceID, Player);
+		const FMovieSceneEvaluationOperand AbsoluteInnerOperand(AbsoluteBinding.GetSequenceID(), AbsoluteBinding.GetGuid());
+		Player.BindingOverrides.Remove(AbsoluteInnerOperand);
 	}
 }
+
+FMovieSceneObjectBindingID FTemplateSequenceSectionTemplate::GetAbsoluteInnerBindingID(FMovieSceneSequenceID LocalSequenceID, IMovieScenePlayer& Player) const
+{
+	// Convert the binding ID that we have, which is local to the child template sequence, into
+	// an absolute ID that fits in the current hierarchy.
+	const FMovieSceneObjectBindingID LocalBinding(InnerOperand.ObjectBindingID, InnerOperand.SequenceID, EMovieSceneObjectBindingSpace::Local);
+
+	const FMovieSceneSequenceHierarchy& Hierarchy = Player.GetEvaluationTemplate().GetHierarchy();
+	return LocalBinding.ResolveLocalToRoot(LocalSequenceID, Hierarchy);
+}
diff --git a/Engine/Plugins/MovieScene/TemplateSequence/Source/TemplateSequence/Private/Tracks/TemplateSequenceTrack.cpp b/Engine/Plugins/MovieScene/TemplateSequence/Source/TemplateSequence/Private/Tracks/TemplateSequenceTrack.cpp
--- a/Engine/Plugins/MovieScene/TemplateSequence/Source/TemplateSequence/Private/Tracks/TemplateSequenceTrack.cpp
+++ b/Engine/Plugins/MovieScene/TemplateSequence/Source/TemplateSequence/Private/Tracks/TemplateSequenceTrack.cpp
@@ -44,6 +44,16 @@ void UTemplateSequenceTrack::PostCompile(FMovieSceneEvaluationTrack& OutTrack, c
 	// Make sure out evaluation template runs before the spawn tracks because it will have to setup the overrides.
 	OutTrack.SetEvaluationGroup(IMovieSceneTracksModule::GetEvaluationGroupName(EBuiltInEvaluationGroup::SpawnObjects));
 	OutTrack.SetEvaluationPriority(GetEvaluationPriority());
+
+	// Cache our parent binding ID onto our templates.
+	for (FMovieSceneEvalTemplatePtr& BaseTemplate : OutTrack.GetChildTemplates())
+	{
+		if (BaseTemplate.IsValid())
+		{
+			FTemplateSequenceSectionTemplate* Template = static_cast<FTemplateSequenceSectionTemplate*>(BaseTemplate.GetPtr());
+			Template->OuterBindingId = Args.ObjectBindingId;
+		}
+	}
 }
 
 #if WITH_EDITORONLY_DATA
diff --git a/Engine/Plugins/MovieScene/TemplateSequence/Source/TemplateSequence/Public/Evaluation/TemplateSequenceSectionTemplate.h b/Engine/Plugins/MovieScene/TemplateSequence/Source/TemplateSequence/Public/Evaluation/TemplateSequenceSectionTemplate.h
--- a/Engine/Plugins/MovieScene/TemplateSequence/Source/TemplateSequence/Public/Evaluation/TemplateSequenceSectionTemplate.h
+++ b/Engine/Plugins/MovieScene/TemplateSequence/Source/TemplateSequence/Public/Evaluation/TemplateSequenceSectionTemplate.h
@@ -21,14 +21,20 @@ struct FTemplateSequenceSectionTemplate : public FMovieSceneEvalTemplate
 	virtual UScriptStruct& GetScriptStructImpl() const override { return *StaticStruct(); }
 	virtual void SetupOverrides() override;
 	virtual void Evaluate(const FMovieSceneEvaluationOperand& Operand, const FMovieSceneContext& Context, const FPersistentEvaluationData& PersistentData, FMovieSceneExecutionTokens& ExecutionTokens) const override;
-	virtual void Initialize(const FMovieSceneEvaluationOperand& Operand, const FMovieSceneContext& Context, FPersistentEvaluationData& PersistentData, IMovieScenePlayer& Player) const override;
+	virtual void Setup(FPersistentEvaluationData& PersistentData, IMovieScenePlayer& Player) const override;
 	virtual void TearDown(FPersistentEvaluationData& PersistentData, IMovieScenePlayer& Player) const override;
 
 private:
+	FMovieSceneObjectBindingID GetAbsoluteInnerBindingID(FMovieSceneSequenceID LocalSequenceID, IMovieScenePlayer& Player) const;
 
 	UPROPERTY()
 	FFrameNumber SectionStartTime;
 
+	UPROPERTY()
+	FGuid OuterBindingId;
+
 	UPROPERTY()
 	FMovieSceneEvaluationOperand InnerOperand;
+
+	friend class UTemplateSequenceTrack;
 };
