Sequencer: Add new camera animation sequence.

A camera animation sequence is a special kind of template sequence that is always bound to a camera (and nothing else). So this change does the following:

- Add a dependency on the cine camera plugin to get access to the CineCameraActor class
- Add options on the template sequence toolkit to suppress the ability to change the root binding type

diff --git a/Engine/Plugins/MovieScene/TemplateSequence/Source/TemplateSequence/Private/CameraAnimationSequence.cpp b/Engine/Plugins/MovieScene/TemplateSequence/Source/TemplateSequence/Private/CameraAnimationSequence.cpp
new file mode 100644
--- /dev/null
+++ b/Engine/Plugins/MovieScene/TemplateSequence/Source/TemplateSequence/Private/CameraAnimationSequence.cpp
@@ -0,0 +1,15 @@
+// Copyright 1998-2019 Epic Games, Inc. All Rights Reserved.
+
+#include "CameraAnimationSequence.h"
+
+UCameraAnimationSequence::UCameraAnimationSequence(const FObjectInitializer& ObjectInitializer)
+    : UTemplateSequence(ObjectInitializer)
+{
+}
+
+#if WITH_EDITOR
+FText UCameraAnimationSequence::GetDisplayName() const
+{
+	return UMovieSceneSequence::GetDisplayName();
+}
+#endif
diff --git a/Engine/Plugins/MovieScene/TemplateSequence/Source/TemplateSequence/Public/CameraAnimationSequence.h b/Engine/Plugins/MovieScene/TemplateSequence/Source/TemplateSequence/Public/CameraAnimationSequence.h
new file mode 100644
--- /dev/null
+++ b/Engine/Plugins/MovieScene/TemplateSequence/Source/TemplateSequence/Public/CameraAnimationSequence.h
@@ -0,0 +1,22 @@
+// Copyright 1998-2019 Epic Games, Inc. All Rights Reserved.
+
+#pragma once
+
+#include "TemplateSequence.h"
+#include "CameraAnimationSequence.generated.h"
+
+/*
+ * A template sequence specifically designed for playing on cameras.
+ */
+UCLASS(BlueprintType)
+class TEMPLATESEQUENCE_API UCameraAnimationSequence : public UTemplateSequence
+{
+public:
+	GENERATED_BODY()
+
+	UCameraAnimationSequence(const FObjectInitializer& ObjectInitializer);
+
+#if WITH_EDITOR
+	virtual FText GetDisplayName() const override;
+#endif
+};
diff --git a/Engine/Plugins/MovieScene/TemplateSequence/Source/TemplateSequenceEditor/Private/AssetTools/CameraAnimationSequenceActions.cpp b/Engine/Plugins/MovieScene/TemplateSequence/Source/TemplateSequenceEditor/Private/AssetTools/CameraAnimationSequenceActions.cpp
new file mode 100644
--- /dev/null
+++ b/Engine/Plugins/MovieScene/TemplateSequence/Source/TemplateSequenceEditor/Private/AssetTools/CameraAnimationSequenceActions.cpp
@@ -0,0 +1,33 @@
+// Copyright 1998-2019 Epic Games, Inc. All Rights Reserved.
+
+#include "CameraAnimationSequenceActions.h"
+#include "CameraAnimationSequence.h"
+#include "CineCameraActor.h"
+#include "TemplateSequenceEditorToolkit.h"
+
+#define LOCTEXT_NAMESPACE "AssetTypeActions"
+
+FCameraAnimationSequenceActions::FCameraAnimationSequenceActions(const TSharedRef<ISlateStyle>& InStyle)
+    : FTemplateSequenceActions(InStyle)
+{
+}
+
+FText FCameraAnimationSequenceActions::GetName() const
+{
+    return NSLOCTEXT("AssetTypeActions", "AssetTypeActions_CameraAnimationSequence", "Camera Animation Sequence");
+}
+
+UClass* FCameraAnimationSequenceActions::GetSupportedClass() const
+{
+    return UCameraAnimationSequence::StaticClass();
+}
+
+void FCameraAnimationSequenceActions::InitializeToolkitParams(FTemplateSequenceToolkitParams& ToolkitParams) const
+{
+    FTemplateSequenceActions::InitializeToolkitParams(ToolkitParams);
+
+    ToolkitParams.bCanChangeBinding = false;
+    ToolkitParams.InitialBindingClass = ACineCameraActor::StaticClass();
+}
+
+#undef LOCTEXT_NAMESPACE
diff --git a/Engine/Plugins/MovieScene/TemplateSequence/Source/TemplateSequenceEditor/Private/AssetTools/CameraAnimationSequenceActions.h b/Engine/Plugins/MovieScene/TemplateSequence/Source/TemplateSequenceEditor/Private/AssetTools/CameraAnimationSequenceActions.h
new file mode 100644
--- /dev/null
+++ b/Engine/Plugins/MovieScene/TemplateSequence/Source/TemplateSequenceEditor/Private/AssetTools/CameraAnimationSequenceActions.h
@@ -0,0 +1,20 @@
+// Copyright 1998-2019 Epic Games, Inc. All Rights Reserved.
+
+#pragma once
+
+#include "TemplateSequenceActions.h"
+
+class FCameraAnimationSequenceActions : public FTemplateSequenceActions
+{
+public:
+    
+    FCameraAnimationSequenceActions(const TSharedRef<ISlateStyle>& InStyle);
+
+    virtual FText GetName() const override;
+    virtual UClass* GetSupportedClass() const override;
+
+protected:
+
+    virtual void InitializeToolkitParams(FTemplateSequenceToolkitParams& ToolkitParams) const override;
+};
+
diff --git a/Engine/Plugins/MovieScene/TemplateSequence/Source/TemplateSequenceEditor/Private/AssetTools/TemplateSequenceActions.cpp b/Engine/Plugins/MovieScene/TemplateSequence/Source/TemplateSequenceEditor/Private/AssetTools/TemplateSequenceActions.cpp
--- a/Engine/Plugins/MovieScene/TemplateSequence/Source/TemplateSequenceEditor/Private/AssetTools/TemplateSequenceActions.cpp
+++ b/Engine/Plugins/MovieScene/TemplateSequence/Source/TemplateSequenceEditor/Private/AssetTools/TemplateSequenceActions.cpp
@@ -60,11 +60,16 @@ void FTemplateSequenceActions::OpenAssetEditor(const TArray<UObject*>& InObjects
 		if (TemplateSequence != nullptr)
 		{
 			TSharedRef<FTemplateSequenceEditorToolkit> Toolkit = MakeShareable(new FTemplateSequenceEditorToolkit(Style));
-			Toolkit->Initialize(Mode, EditWithinLevelEditor, TemplateSequence);
+			FTemplateSequenceToolkitParams ToolkitParams;
+			InitializeToolkitParams(ToolkitParams);
+			Toolkit->Initialize(Mode, EditWithinLevelEditor, TemplateSequence, ToolkitParams);
 		}
 	}
 }
 
+void FTemplateSequenceActions::InitializeToolkitParams(FTemplateSequenceToolkitParams& ToolkitParams) const
+{
+}
 
 bool FTemplateSequenceActions::ShouldForceWorldCentric()
 {
diff --git a/Engine/Plugins/MovieScene/TemplateSequence/Source/TemplateSequenceEditor/Private/AssetTools/TemplateSequenceActions.h b/Engine/Plugins/MovieScene/TemplateSequence/Source/TemplateSequenceEditor/Private/AssetTools/TemplateSequenceActions.h
--- a/Engine/Plugins/MovieScene/TemplateSequence/Source/TemplateSequenceEditor/Private/AssetTools/TemplateSequenceActions.h
+++ b/Engine/Plugins/MovieScene/TemplateSequence/Source/TemplateSequenceEditor/Private/AssetTools/TemplateSequenceActions.h
@@ -7,6 +7,7 @@
 #include "Toolkits/IToolkitHost.h"
 #include "AssetTypeActions_Base.h"
 
+struct FTemplateSequenceToolkitParams;
 
 /**
  * Implements actions for UTemplateSequence assets.
@@ -34,6 +35,10 @@ public:
 	virtual bool CanLocalize() const override { return false; }
 	virtual bool HasActions(const TArray<UObject*>& InObjects) const override { return false; }
 
+protected:
+
+	virtual void InitializeToolkitParams(FTemplateSequenceToolkitParams& ToolkitParams) const;
+
 private:
 
 	/** Pointer to the style set to use for toolkits. */
diff --git a/Engine/Plugins/MovieScene/TemplateSequence/Source/TemplateSequenceEditor/Private/Factories/CameraAnimationSequenceFactoryNew.cpp b/Engine/Plugins/MovieScene/TemplateSequence/Source/TemplateSequenceEditor/Private/Factories/CameraAnimationSequenceFactoryNew.cpp
new file mode 100644
--- /dev/null
+++ b/Engine/Plugins/MovieScene/TemplateSequence/Source/TemplateSequenceEditor/Private/Factories/CameraAnimationSequenceFactoryNew.cpp
@@ -0,0 +1,38 @@
+// Copyright 1998-2019 Epic Games, Inc. All Rights Reserved.
+
+#include "Factories/CameraAnimationSequenceFactoryNew.h"
+#include "CameraAnimationSequence.h"
+#include "MovieScene.h"
+#include "MovieSceneToolsProjectSettings.h"
+
+#define LOCTEXT_NAMESPACE "MovieSceneFactory"
+
+UCameraAnimationSequenceFactoryNew::UCameraAnimationSequenceFactoryNew(const FObjectInitializer& ObjectInitializer)
+	: Super(ObjectInitializer)
+{
+	bCreateNew = true;
+	bEditAfterNew = true;
+	SupportedClass = UCameraAnimationSequence::StaticClass();
+}
+
+UObject* UCameraAnimationSequenceFactoryNew::FactoryCreateNew(UClass* Class, UObject* InParent, FName Name, EObjectFlags Flags, UObject* Context, FFeedbackContext* Warn)
+{
+	auto NewCameraAnimationSequence = NewObject<UCameraAnimationSequence>(InParent, Name, Flags | RF_Transactional);
+	NewCameraAnimationSequence->Initialize();
+
+	// Set up some sensible defaults
+	const UMovieSceneToolsProjectSettings* ProjectSettings = GetDefault<UMovieSceneToolsProjectSettings>();
+
+	FFrameRate TickResolution = NewCameraAnimationSequence->GetMovieScene()->GetTickResolution();
+	NewCameraAnimationSequence->GetMovieScene()->SetPlaybackRange((ProjectSettings->DefaultStartTime*TickResolution).FloorToFrame(), (ProjectSettings->DefaultDuration*TickResolution).FloorToFrame().Value);
+
+	return NewCameraAnimationSequence;
+}
+
+bool UCameraAnimationSequenceFactoryNew::ShouldShowInNewMenu() const
+{
+	return true;
+}
+
+#undef LOCTEXT_NAMESPACE
+
diff --git a/Engine/Plugins/MovieScene/TemplateSequence/Source/TemplateSequenceEditor/Private/Factories/CameraAnimationSequenceFactoryNew.h b/Engine/Plugins/MovieScene/TemplateSequence/Source/TemplateSequenceEditor/Private/Factories/CameraAnimationSequenceFactoryNew.h
new file mode 100644
--- /dev/null
+++ b/Engine/Plugins/MovieScene/TemplateSequence/Source/TemplateSequenceEditor/Private/Factories/CameraAnimationSequenceFactoryNew.h
@@ -0,0 +1,23 @@
+// Copyright 1998-2019 Epic Games, Inc. All Rights Reserved.
+
+#pragma once
+
+#include "CoreMinimal.h"
+#include "UObject/ObjectMacros.h"
+#include "Factories/Factory.h"
+#include "CameraAnimationSequenceFactoryNew.generated.h"
+
+/**
+ * Implements a factory for UCameraAnimationSequence objects.
+ */
+UCLASS(hidecategories = Object)
+class UCameraAnimationSequenceFactoryNew : public UFactory
+{
+	GENERATED_UCLASS_BODY()
+
+public:
+	// UFactory Interface
+	virtual UObject* FactoryCreateNew(UClass* Class, UObject* InParent, FName Name, EObjectFlags Flags, UObject* Context, FFeedbackContext* Warn) override;
+	virtual bool ShouldShowInNewMenu() const override;
+};
+
diff --git a/Engine/Plugins/MovieScene/TemplateSequence/Source/TemplateSequenceEditor/Private/TemplateSequenceEditorModule.cpp b/Engine/Plugins/MovieScene/TemplateSequence/Source/TemplateSequenceEditor/Private/TemplateSequenceEditorModule.cpp
--- a/Engine/Plugins/MovieScene/TemplateSequence/Source/TemplateSequenceEditor/Private/TemplateSequenceEditorModule.cpp
+++ b/Engine/Plugins/MovieScene/TemplateSequence/Source/TemplateSequenceEditor/Private/TemplateSequenceEditorModule.cpp
@@ -1,5 +1,6 @@
 // Copyright 1998-2019 Epic Games, Inc. All Rights Reserved.
 
+#include "AssetTools/CameraAnimationSequenceActions.h"
 #include "AssetTools/TemplateSequenceActions.h"
 #include "AssetToolsModule.h"
 #include "ISequencerModule.h"
@@ -89,8 +90,12 @@ private:
 	void RegisterAssetTools()
 	{
 		IAssetTools& AssetTools = FModuleManager::LoadModuleChecked<FAssetToolsModule>("AssetTools").Get();
+
 		TemplateSequenceTypeActions = MakeShared<FTemplateSequenceActions>(FTemplateSequenceEditorStyle::Get());
 		AssetTools.RegisterAssetTypeActions(TemplateSequenceTypeActions.ToSharedRef());
+
+		CameraAnimationSequenceActions = MakeShared<FCameraAnimationSequenceActions>(FTemplateSequenceEditorStyle::Get());
+		AssetTools.RegisterAssetTypeActions(CameraAnimationSequenceActions.ToSharedRef());
 	}
 
 	void UnregisterAssetTools()
@@ -101,6 +106,7 @@ private:
 		{
 			IAssetTools& AssetTools = AssetToolsModule->Get();
 			AssetTools.UnregisterAssetTypeActions(TemplateSequenceTypeActions.ToSharedRef());
+			AssetTools.UnregisterAssetTypeActions(CameraAnimationSequenceActions.ToSharedRef());
 		}
 	}
 
@@ -137,6 +143,7 @@ private:
 private:
 
 	TSharedPtr<FTemplateSequenceActions> TemplateSequenceTypeActions;
+	TSharedPtr<FCameraAnimationSequenceActions> CameraAnimationSequenceActions;
 
 	FDelegateHandle SequenceEditorHandle;
 	FDelegateHandle TemplateSequenceTrackCreateEditorHandle;
diff --git a/Engine/Plugins/MovieScene/TemplateSequence/Source/TemplateSequenceEditor/Private/TemplateSequenceEditorToolkit.cpp b/Engine/Plugins/MovieScene/TemplateSequence/Source/TemplateSequenceEditor/Private/TemplateSequenceEditorToolkit.cpp
--- a/Engine/Plugins/MovieScene/TemplateSequence/Source/TemplateSequenceEditor/Private/TemplateSequenceEditorToolkit.cpp
+++ b/Engine/Plugins/MovieScene/TemplateSequence/Source/TemplateSequenceEditor/Private/TemplateSequenceEditorToolkit.cpp
@@ -60,7 +60,7 @@ FTemplateSequenceEditorToolkit::~FTemplateSequenceEditorToolkit()
 	});
 }
 
-void FTemplateSequenceEditorToolkit::Initialize(const EToolkitMode::Type Mode, const TSharedPtr<IToolkitHost>& InitToolkitHost, UTemplateSequence* InTemplateSequence)
+void FTemplateSequenceEditorToolkit::Initialize(const EToolkitMode::Type Mode, const TSharedPtr<IToolkitHost>& InitToolkitHost, UTemplateSequence* InTemplateSequence, const FTemplateSequenceToolkitParams& ToolkitParams)
 {
 	// create tab layout
 	const TSharedRef<FTabManager::FLayout> StandaloneDefaultLayout = FTabManager::NewLayout("Standalone_TemplateSequenceEditor")
@@ -100,19 +100,37 @@ void FTemplateSequenceEditorToolkit::Initialize(const EToolkitMode::Type Mode, c
 		SequencerInitParams.ViewParams.UniqueName = "TemplateSequenceEditor";
 		SequencerInitParams.ViewParams.ScrubberStyle = ESequencerScrubberStyle::FrameBlock;
 		SequencerInitParams.ViewParams.OnReceivedFocus.BindRaw(this, &FTemplateSequenceEditorToolkit::OnSequencerReceivedFocus);
-		SequencerInitParams.ViewParams.OnAssetsDrop.BindRaw(this, &FTemplateSequenceEditorToolkit::OnSequencerAssetsDrop);
-		SequencerInitParams.ViewParams.OnClassesDrop.BindRaw(this, &FTemplateSequenceEditorToolkit::OnSequencerClassesDrop);
-		SequencerInitParams.ViewParams.OnActorsDrop.BindRaw(this, &FTemplateSequenceEditorToolkit::OnSequencerActorsDrop);
 
-		TSharedRef<FExtender> ToolbarExtender = MakeShared<FExtender>();
-		ToolbarExtender->AddToolBarExtension("Base Commands", EExtensionHook::After, nullptr, FToolBarExtensionDelegate::CreateSP(this, &FTemplateSequenceEditorToolkit::ExtendSequencerToolbar));
-		SequencerInitParams.ViewParams.ToolbarExtender = ToolbarExtender;
+		if (ToolkitParams.bCanChangeBinding)
+		{
+			// Callbacks for changing the root binding by drag-and-dropping stuff on the Sequencer.
+			SequencerInitParams.ViewParams.OnAssetsDrop.BindRaw(this, &FTemplateSequenceEditorToolkit::OnSequencerAssetsDrop);
+			SequencerInitParams.ViewParams.OnClassesDrop.BindRaw(this, &FTemplateSequenceEditorToolkit::OnSequencerClassesDrop);
+			SequencerInitParams.ViewParams.OnActorsDrop.BindRaw(this, &FTemplateSequenceEditorToolkit::OnSequencerActorsDrop);
+
+			// Extended for showing toolbar controls to change the root binding.
+			TSharedRef<FExtender> ToolbarExtender = MakeShared<FExtender>();
+			ToolbarExtender->AddToolBarExtension("Base Commands", EExtensionHook::After, nullptr, FToolBarExtensionDelegate::CreateSP(this, &FTemplateSequenceEditorToolkit::ExtendSequencerToolbar));
+			SequencerInitParams.ViewParams.ToolbarExtender = ToolbarExtender;
+		}
+		else
+		{
+			// Bind drag-and-drop callbacks to suppress them from doing anything.
+			SequencerInitParams.ViewParams.OnAssetsDrop.BindLambda([](const TArray<UObject*>&, const FAssetDragDropOp&) -> bool { return false; });
+			SequencerInitParams.ViewParams.OnClassesDrop.BindLambda([](const TArray<TWeakObjectPtr<UClass>>&, const FClassDragDropOp&) -> bool { return false; });
+			SequencerInitParams.ViewParams.OnActorsDrop.BindLambda([](const TArray<TWeakObjectPtr<AActor>>&, const FActorDragDropGraphEdOp&) -> bool { return false; });
+		}
 	}
 
 	Sequencer = FModuleManager::LoadModuleChecked<ISequencerModule>("Sequencer").CreateSequencer(SequencerInitParams);
 
 	Sequencer->OnActorAddedToSequencer().AddSP(this, &FTemplateSequenceEditorToolkit::HandleActorAddedToSequencer);
 
+	if (ToolkitParams.InitialBindingClass != nullptr)
+	{
+		ChangeActorBinding(*ToolkitParams.InitialBindingClass);
+	}
+
 	FLevelEditorSequencerIntegrationOptions Options;
 	Options.bRequiresLevelEvents = true;
 	Options.bRequiresActorEvents = true;
@@ -382,6 +400,8 @@ bool FTemplateSequenceEditorToolkit::OnSequencerAssetsDrop(const TArray<UObject*
 		// TODO: check for dropping a sequence?
 
 		ChangeActorBinding(*CurObject, DragDropOp.GetActorFactory());
+
+		return true;
 	}
 
 	return true;
@@ -389,6 +409,15 @@ bool FTemplateSequenceEditorToolkit::OnSequencerAssetsDrop(const TArray<UObject*
 
 bool FTemplateSequenceEditorToolkit::OnSequencerClassesDrop(const TArray<TWeakObjectPtr<UClass>>& Classes, const FClassDragDropOp& DragDropOp)
 {
+	if (Classes.Num() > 0 && Classes[0].IsValid())
+	{
+		// Only drop the first class.
+		UClass* CurClass = Classes[0].Get();
+
+		ChangeActorBinding(*CurClass);
+
+		return true;
+	}
 	return false;
 }
 
diff --git a/Engine/Plugins/MovieScene/TemplateSequence/Source/TemplateSequenceEditor/Private/TemplateSequenceEditorToolkit.h b/Engine/Plugins/MovieScene/TemplateSequence/Source/TemplateSequenceEditor/Private/TemplateSequenceEditorToolkit.h
--- a/Engine/Plugins/MovieScene/TemplateSequence/Source/TemplateSequenceEditor/Private/TemplateSequenceEditorToolkit.h
+++ b/Engine/Plugins/MovieScene/TemplateSequence/Source/TemplateSequenceEditor/Private/TemplateSequenceEditorToolkit.h
@@ -15,6 +15,12 @@ class FAssetDragDropOp;
 class FClassDragDropOp;
 class FActorDragDropGraphEdOp;
 
+struct FTemplateSequenceToolkitParams
+{
+	bool bCanChangeBinding = true;
+	UClass* InitialBindingClass = nullptr;
+};
+
 /**
  * Implements an Editor toolkit for template sequences.
  */
@@ -42,7 +48,7 @@ public:
 	 * @param TemplateSequence The animation to edit.
 	 * @param TrackEditorDelegates Delegates to call to create auto-key handlers for this sequencer.
 	 */
-	void Initialize(const EToolkitMode::Type Mode, const TSharedPtr<IToolkitHost>& InitToolkitHost, UTemplateSequence* TemplateSequence);
+	void Initialize(const EToolkitMode::Type Mode, const TSharedPtr<IToolkitHost>& InitToolkitHost, UTemplateSequence* TemplateSequence, const FTemplateSequenceToolkitParams& ToolkitParams);
 
 public:
 
diff --git a/Engine/Plugins/MovieScene/TemplateSequence/Source/TemplateSequenceEditor/Private/TrackEditors/TemplateSequenceTrackEditor.cpp b/Engine/Plugins/MovieScene/TemplateSequence/Source/TemplateSequenceEditor/Private/TrackEditors/TemplateSequenceTrackEditor.cpp
--- a/Engine/Plugins/MovieScene/TemplateSequence/Source/TemplateSequenceEditor/Private/TrackEditors/TemplateSequenceTrackEditor.cpp
+++ b/Engine/Plugins/MovieScene/TemplateSequence/Source/TemplateSequenceEditor/Private/TrackEditors/TemplateSequenceTrackEditor.cpp
@@ -3,11 +3,14 @@
 #include "TrackEditors/TemplateSequenceTrackEditor.h"
 #include "AssetData.h"
 #include "AssetRegistryModule.h"
+#include "Camera/CameraComponent.h"
+#include "CameraAnimationSequence.h"
 #include "CommonMovieSceneTools.h"
 #include "ContentBrowserModule.h"
 #include "Framework/Application/SlateApplication.h"
 #include "Framework/MultiBox/MultiBoxBuilder.h"
 #include "IContentBrowserSingleton.h"
+#include "MovieSceneCommonHelpers.h"
 #include "MovieSceneTimeHelpers.h"
 #include "Sections/TemplateSequenceSection.h"
 #include "SequencerSectionPainter.h"
@@ -38,6 +41,18 @@ void FTemplateSequenceTrackEditor::BuildObjectBindingTrackMenu(FMenuBuilder& Men
 		LOCTEXT("AddTemplateSequence", "Template Sequence"), NSLOCTEXT("Sequencer", "AddTemplateSequenceTooltip", "Adds a track that can play a template sequence asset using the parent binding."),
 		FNewMenuDelegate::CreateRaw(this, &FTemplateSequenceTrackEditor::AddTemplateSequenceSubMenu, ObjectBindings)
 	);
+
+	if (ObjectBindings.Num() > 0)
+	{
+		UCameraComponent const* const CameraComponent = AcquireCameraComponentFromObjectGuid(ObjectBindings[0]);
+		if (CameraComponent != nullptr)
+		{
+			MenuBuilder.AddSubMenu(
+				LOCTEXT("AddCameraAnimationSequence", "Camera Animation"), NSLOCTEXT("Sequencer", "AddCameraAnimationSequenceTooltip", "Adds a camera animation template sequence on the parent binding."),
+				FNewMenuDelegate::CreateRaw(this, &FTemplateSequenceTrackEditor::AddCameraAnimationSequenceSubMenu, ObjectBindings)
+			);
+		}
+	}
 }
 
 TSharedRef<ISequencerSection> FTemplateSequenceTrackEditor::MakeSectionInterface(UMovieSceneSection& SectionObject, UMovieSceneTrack& Track, FGuid ObjectBinding)
@@ -49,13 +64,26 @@ TSharedRef<ISequencerSection> FTemplateSequenceTrackEditor::MakeSectionInterface
 
 void FTemplateSequenceTrackEditor::AddTemplateSequenceSubMenu(FMenuBuilder& MenuBuilder, TArray<FGuid> ObjectBindings)
 {
+	AddTemplateSequenceAssetSubMenu(MenuBuilder, ObjectBindings, UTemplateSequence::StaticClass());
+}
+
+void FTemplateSequenceTrackEditor::AddCameraAnimationSequenceSubMenu(FMenuBuilder& MenuBuilder, TArray<FGuid> ObjectBindings)
+{
+	AddTemplateSequenceAssetSubMenu(MenuBuilder, ObjectBindings, UCameraAnimationSequence::StaticClass());
+}
+
+void FTemplateSequenceTrackEditor::AddTemplateSequenceAssetSubMenu(FMenuBuilder& MenuBuilder, TArray<FGuid> ObjectBindings, const UClass* TemplateSequenceClass, bool bRecursiveClasses)
+{
+	ensure(TemplateSequenceClass != nullptr && TemplateSequenceClass->IsChildOf<UTemplateSequence>());
+
 	FAssetPickerConfig AssetPickerConfig;
 	{
 		AssetPickerConfig.OnAssetSelected = FOnAssetSelected::CreateRaw(this, &FTemplateSequenceTrackEditor::OnTemplateSequenceAssetSelected, ObjectBindings);
 		AssetPickerConfig.OnAssetEnterPressed = FOnAssetEnterPressed::CreateRaw(this, &FTemplateSequenceTrackEditor::OnTemplateSequenceAssetEnterPressed, ObjectBindings);
 		AssetPickerConfig.bAllowNullSelection = false;
 		AssetPickerConfig.InitialAssetViewType = EAssetViewType::List;
-		AssetPickerConfig.Filter.ClassNames.Add(UTemplateSequence::StaticClass()->GetFName());
+		AssetPickerConfig.Filter.ClassNames.Add(TemplateSequenceClass->GetFName());
+		AssetPickerConfig.Filter.bRecursiveClasses = bRecursiveClasses;
 	}
 
 	FContentBrowserModule& ContentBrowserModule = FModuleManager::Get().LoadModuleChecked<FContentBrowserModule>(TEXT("ContentBrowser"));
@@ -132,6 +160,33 @@ FKeyPropertyResult FTemplateSequenceTrackEditor::AddKeyInternal(FFrameNumber Key
 	return KeyPropertyResult;
 }
 
+UCameraComponent* FTemplateSequenceTrackEditor::AcquireCameraComponentFromObjectGuid(const FGuid& Guid)
+{
+	USkeleton* Skeleton = nullptr;
+	for (TWeakObjectPtr<> WeakObject : GetSequencer()->FindObjectsInCurrentSequence(Guid))
+	{
+		UObject* const Obj = WeakObject.Get();
+	
+		if (AActor* const Actor = Cast<AActor>(Obj))
+		{
+			UCameraComponent* const CameraComp = MovieSceneHelpers::CameraComponentFromActor(Actor);
+			if (CameraComp)
+			{
+				return CameraComp;
+			}
+		}
+		else if (UCameraComponent* const CameraComp = Cast<UCameraComponent>(Obj))
+		{
+			if (CameraComp->IsActive())
+			{
+				return CameraComp;
+			}
+		}
+	}
+
+	return nullptr;
+}
+
 FTemplateSequenceSection::FTemplateSequenceSection(TSharedPtr<ISequencer> InSequencer, UTemplateSequenceSection& InSection)
 	: TSubSectionMixin(InSequencer, InSection)
 {
diff --git a/Engine/Plugins/MovieScene/TemplateSequence/Source/TemplateSequenceEditor/Private/TrackEditors/TemplateSequenceTrackEditor.h b/Engine/Plugins/MovieScene/TemplateSequence/Source/TemplateSequenceEditor/Private/TrackEditors/TemplateSequenceTrackEditor.h
--- a/Engine/Plugins/MovieScene/TemplateSequence/Source/TemplateSequenceEditor/Private/TrackEditors/TemplateSequenceTrackEditor.h
+++ b/Engine/Plugins/MovieScene/TemplateSequence/Source/TemplateSequenceEditor/Private/TrackEditors/TemplateSequenceTrackEditor.h
@@ -9,6 +9,7 @@
 
 struct FAssetData;
 class FMenuBuilder;
+class UCameraComponent;
 class UTemplateSequence;
 class UTemplateSequenceSection;
 
@@ -27,10 +28,15 @@ public:
 
 private:
 	void AddTemplateSequenceSubMenu(FMenuBuilder& MenuBuilder, TArray<FGuid> ObjectBindings);
+	void AddCameraAnimationSequenceSubMenu(FMenuBuilder& MenuBuilder, TArray<FGuid> ObjectBindings);
+	void AddTemplateSequenceAssetSubMenu(FMenuBuilder& MenuBuilder, TArray<FGuid> ObjectBindings, const UClass* TemplateSequenceClass, bool bRecursiveClasses = false);
+
 	void OnTemplateSequenceAssetSelected(const FAssetData& AssetData, TArray<FGuid> ObjectBindings);
 	void OnTemplateSequenceAssetEnterPressed(const TArray<FAssetData>& AssetData, TArray<FGuid> ObjectBindings);
 
 	FKeyPropertyResult AddKeyInternal(FFrameNumber KeyTime, TArray<FGuid> ObjectBindings, UTemplateSequence* TemplateSequence);
+
+	UCameraComponent* AcquireCameraComponentFromObjectGuid(const FGuid& Guid);
 };
 
 class FTemplateSequenceSection
diff --git a/Engine/Plugins/MovieScene/TemplateSequence/Source/TemplateSequenceEditor/TemplateSequenceEditor.Build.cs b/Engine/Plugins/MovieScene/TemplateSequence/Source/TemplateSequenceEditor/TemplateSequenceEditor.Build.cs
--- a/Engine/Plugins/MovieScene/TemplateSequence/Source/TemplateSequenceEditor/TemplateSequenceEditor.Build.cs
+++ b/Engine/Plugins/MovieScene/TemplateSequence/Source/TemplateSequenceEditor/TemplateSequenceEditor.Build.cs
@@ -10,6 +10,7 @@ public class TemplateSequenceEditor : ModuleRules
 			new string[] {
 				"TemplateSequence",
 				"BlueprintGraph",
+				"CinematicCamera",
 				"Core",
 				"CoreUObject",
 				"Engine",
