Sequencer: Fix static analysis warnings.

diff --git a/Engine/Plugins/MovieScene/TemplateSequence/Source/TemplateSequenceEditor/Private/TrackEditors/TemplateSequenceTrackEditor.cpp b/Engine/Plugins/MovieScene/TemplateSequence/Source/TemplateSequenceEditor/Private/TrackEditors/TemplateSequenceTrackEditor.cpp
--- a/Engine/Plugins/MovieScene/TemplateSequence/Source/TemplateSequenceEditor/Private/TrackEditors/TemplateSequenceTrackEditor.cpp
+++ b/Engine/Plugins/MovieScene/TemplateSequence/Source/TemplateSequenceEditor/Private/TrackEditors/TemplateSequenceTrackEditor.cpp
@@ -74,7 +74,10 @@ void FTemplateSequenceTrackEditor::AddCameraAnimationSequenceSubMenu(FMenuBuilde
 
 void FTemplateSequenceTrackEditor::AddTemplateSequenceAssetSubMenu(FMenuBuilder& MenuBuilder, TArray<FGuid> ObjectBindings, const UClass* TemplateSequenceClass, bool bRecursiveClasses)
 {
-	ensure(TemplateSequenceClass != nullptr && TemplateSequenceClass->IsChildOf<UTemplateSequence>());
+	if (!ensure(TemplateSequenceClass != nullptr && TemplateSequenceClass->IsChildOf<UTemplateSequence>()))
+	{
+		return;
+	}
 
 	FAssetPickerConfig AssetPickerConfig;
 	{
