Sequencer: Don't modify a template sequence when changing the bound actor class to the same thing it already had.

diff --git a/Engine/Plugins/MovieScene/TemplateSequence/Source/TemplateSequenceEditor/Private/TemplateSequenceEditorToolkit.cpp b/Engine/Plugins/MovieScene/TemplateSequence/Source/TemplateSequenceEditor/Private/TemplateSequenceEditorToolkit.cpp
--- a/Engine/Plugins/MovieScene/TemplateSequence/Source/TemplateSequenceEditor/Private/TemplateSequenceEditorToolkit.cpp
+++ b/Engine/Plugins/MovieScene/TemplateSequence/Source/TemplateSequenceEditor/Private/TemplateSequenceEditorToolkit.cpp
@@ -428,13 +428,24 @@ bool FTemplateSequenceEditorToolkit::OnSequencerActorsDrop(const TArray<TWeakObj
 
 void FTemplateSequenceEditorToolkit::ChangeActorBinding(UObject& Object, UActorFactory* ActorFactory, bool bSetupDefaults)
 {
+	// See if we have anything to do in the first place.
+	if (UClass* ChosenClass = Cast<UClass>(&Object))
+	{
+		if (ChosenClass == TemplateSequence->BoundActorClass)
+		{
+			return;
+		}
+	}
+
 	UMovieScene* MovieScene = TemplateSequence->GetMovieScene();
 	check(MovieScene != nullptr);
 
 	// See if we previously had a main object binding.
 	FGuid PreviousSpawnableGuid;
 	if (MovieScene->GetSpawnableCount() > 0)
+	{
 		PreviousSpawnableGuid = MovieScene->GetSpawnable(0).GetGuid();
+	}
 
 	// Make the new spawnable object binding.
 	FGuid NewSpawnableGuid = Sequencer->MakeNewSpawnable(Object, ActorFactory, bSetupDefaults);
