Sequencer: Add bound actor class information in asset registry tags.

diff --git a/Engine/Plugins/MovieScene/TemplateSequence/Source/TemplateSequence/Private/TemplateSequence.cpp b/Engine/Plugins/MovieScene/TemplateSequence/Source/TemplateSequence/Private/TemplateSequence.cpp
--- a/Engine/Plugins/MovieScene/TemplateSequence/Source/TemplateSequence/Private/TemplateSequence.cpp
+++ b/Engine/Plugins/MovieScene/TemplateSequence/Source/TemplateSequence/Private/TemplateSequence.cpp
@@ -103,4 +103,32 @@ FText UTemplateSequence::GetDisplayName() const
 {
 	return UMovieSceneSequence::GetDisplayName();
 }
+
+void UTemplateSequence::GetAssetRegistryTags(TArray<FAssetRegistryTag>& OutTags) const
+{
+	Super::GetAssetRegistryTags(OutTags);
+
+	if (BoundActorClass != nullptr)
+	{
+		FAssetRegistryTag Tag("BoundActorClass", BoundActorClass->GetName(), FAssetRegistryTag::TT_Alphabetical);
+		OutTags.Add(Tag);
+	}
+	else
+	{
+		OutTags.Emplace("BoundActorClass", "(None)", FAssetRegistryTag::TT_Alphabetical);
+	}
+}
+
+void UTemplateSequence::GetAssetRegistryTagMetadata(TMap<FName, FAssetRegistryTagMetadata>& OutMetadata) const
+{
+	Super::GetAssetRegistryTagMetadata(OutMetadata);
+
+	OutMetadata.Add(
+		"BoundActorClass",
+		FAssetRegistryTagMetadata()
+			.SetDisplayName(NSLOCTEXT("TemplateSequence", "BoundActorClass_Label", "Bound Actor Class"))
+			.SetTooltip(NSLOCTEXT("TemplateSequence", "BoundActorClass_Tooltip", "The type of actor bound to this template sequence"))
+		);
+}
+
 #endif
diff --git a/Engine/Plugins/MovieScene/TemplateSequence/Source/TemplateSequence/Public/TemplateSequence.h b/Engine/Plugins/MovieScene/TemplateSequence/Source/TemplateSequence/Public/TemplateSequence.h
--- a/Engine/Plugins/MovieScene/TemplateSequence/Source/TemplateSequence/Public/TemplateSequence.h
+++ b/Engine/Plugins/MovieScene/TemplateSequence/Source/TemplateSequence/Public/TemplateSequence.h
@@ -34,6 +34,9 @@ public:
 
 #if WITH_EDITOR
 	virtual FText GetDisplayName() const override;
+
+	virtual void GetAssetRegistryTagMetadata(TMap<FName, FAssetRegistryTagMetadata>& OutMetadata) const override;
+	virtual void GetAssetRegistryTags(TArray<FAssetRegistryTag>& OutTags) const override;
 #endif
 
 public:
