Fixed deprecation warning

diff --git a/Engine/Plugins/MovieScene/TemplateSequence/Source/TemplateSequenceEditor/Private/TrackEditors/TemplateSequenceTrackEditor.cpp b/Engine/Plugins/MovieScene/TemplateSequence/Source/TemplateSequenceEditor/Private/TrackEditors/TemplateSequenceTrackEditor.cpp
--- a/Engine/Plugins/MovieScene/TemplateSequence/Source/TemplateSequenceEditor/Private/TrackEditors/TemplateSequenceTrackEditor.cpp
+++ b/Engine/Plugins/MovieScene/TemplateSequence/Source/TemplateSequenceEditor/Private/TrackEditors/TemplateSequenceTrackEditor.cpp
@@ -129,9 +129,8 @@ public:
 		Filter.ClassNames.Add(BaseClass->GetFName());
 
 		IAssetRegistry& AssetRegistry = FAssetRegistryModule::GetRegistry();
-		FARFilter ExpandedFilter;
-		AssetRegistry.ExpandRecursiveFilter(Filter, ExpandedFilter);
-		TSet<FName> ChildClassNames(ExpandedFilter.ClassNames);
+		FARCompiledFilter CompiledFilter;
+		AssetRegistry.CompileFilter(Filter, CompiledFilter);
 
 		FAssetPickerConfig AssetPickerConfig;
 		{
@@ -150,7 +149,7 @@ public:
 			}
 
 			AssetPickerConfig.OnShouldFilterAsset = FOnShouldFilterAsset::CreateLambda(
-				[this, ChildClassNames](const FAssetData& AssetData) -> bool
+				[this, ChildClassNames = CompiledFilter.ClassNames](const FAssetData& AssetData) -> bool
 				{
 					if (LegacyBaseClass == nullptr || AssetData.AssetClass != LegacyBaseClass->GetFName())
 					{
