Sequencer: Don't make cyclic inclusions with template sequences, fix a possible loop with template sequence track editor.

- Refactor the sub track editor to put `CanAddSubSequence` into a common place.

- We don't need to auto-create the object binding when adding a template sequence section, because this can currently only be done via the object binidng's menu anyway. This solves a bug with a template sequence track to a template sequence that's currently itself in a template sequence track (are you following? :) ). The problem was that the template sequence's root binding would, at that point, be bound to a binding override, so trying to find the object binding for that object would fail in the template sequence.

- Add warnings/notifications when adding a template sequence track that points to a template sequence with a different tick resolution.

- When creating new template sequence sections, select them all in one go instead of one by one (and ending up with only the last one selected).

diff --git a/Engine/Plugins/MovieScene/TemplateSequence/Source/TemplateSequenceEditor/Private/TrackEditors/TemplateSequenceTrackEditor.cpp b/Engine/Plugins/MovieScene/TemplateSequence/Source/TemplateSequenceEditor/Private/TrackEditors/TemplateSequenceTrackEditor.cpp
--- a/Engine/Plugins/MovieScene/TemplateSequence/Source/TemplateSequenceEditor/Private/TrackEditors/TemplateSequenceTrackEditor.cpp
+++ b/Engine/Plugins/MovieScene/TemplateSequence/Source/TemplateSequenceEditor/Private/TrackEditors/TemplateSequenceTrackEditor.cpp
@@ -354,24 +354,40 @@ FKeyPropertyResult FTemplateSequenceTrackEditor::AddKeyInternal(FFrameNumber Key
 {
 	FKeyPropertyResult KeyPropertyResult;
 
-	bool bHandleCreated = false;
-	bool bTrackCreated = false;
-	bool bTrackModified = false;
+	if (TemplateSequence->GetMovieScene()->GetPlaybackRange().IsEmpty())
+	{
+		FNotificationInfo Info(FText::Format(LOCTEXT("InvalidSequenceDuration", "Invalid template sequence {0}. The template sequence has no duration."), TemplateSequence->GetDisplayName()));
+		Info.bUseLargeFont = false;
+		FSlateNotificationManager::Get().AddNotification(Info);
+		return KeyPropertyResult;
+	}
+
+	if (!CanAddSubSequence(*TemplateSequence))
+	{
+		FNotificationInfo Info(FText::Format(LOCTEXT("InvalidSequenceCycle", "Invalid template sequence {0}. There could be a circular dependency."), TemplateSequence->GetDisplayName()));
+		Info.bUseLargeFont = false;
+		FSlateNotificationManager::Get().AddNotification(Info);
+		return KeyPropertyResult;
+	}
 
+	TArray<UMovieSceneSection*> NewSections;
 	TSharedPtr<ISequencer> SequencerPtr = GetSequencer();
 	if (SequencerPtr.IsValid())
 	{
-		for (const FGuid& ObjectBindingGuid : ObjectBindings)
+		const FFrameRate OuterTickResolution = SequencerPtr->GetFocusedTickResolution();
+		const FFrameRate SubTickResolution = TemplateSequence->GetMovieScene()->GetTickResolution();
+		if (SubTickResolution != OuterTickResolution)
 		{
-			UObject* Object = SequencerPtr->FindSpawnedObjectOrTemplate(ObjectBindingGuid);
-
-			FFindOrCreateHandleResult HandleResult = FindOrCreateHandleToObject(Object);
-			FGuid ObjectHandle = HandleResult.Handle;
-			KeyPropertyResult.bHandleCreated |= HandleResult.bWasCreated;
+			FNotificationInfo Info(FText::Format(LOCTEXT("TickResolutionMismatch", "The parent sequence has a different tick resolution {0} than the newly added sequence {1}"), OuterTickResolution.ToPrettyText(), SubTickResolution.ToPrettyText()));
+			Info.bUseLargeFont = false;
+			FSlateNotificationManager::Get().AddNotification(Info);
+		}
 
-			if (ObjectHandle.IsValid())
+		for (const FGuid& ObjectBindingGuid : ObjectBindings)
+		{
+			if (ObjectBindingGuid.IsValid())
 			{
-				FFindOrCreateTrackResult TrackResult = FindOrCreateTrackForObject(ObjectHandle, UTemplateSequenceTrack::StaticClass());
+				FFindOrCreateTrackResult TrackResult = FindOrCreateTrackForObject(ObjectBindingGuid, UTemplateSequenceTrack::StaticClass());
 				UMovieSceneTrack* Track = TrackResult.Track;
 				KeyPropertyResult.bTrackCreated |= TrackResult.bWasCreated;
 
@@ -379,15 +395,22 @@ FKeyPropertyResult FTemplateSequenceTrackEditor::AddKeyInternal(FFrameNumber Key
 				{
 					UMovieSceneSection* NewSection = Cast<UTemplateSequenceTrack>(Track)->AddNewTemplateSequenceSection(KeyTime, TemplateSequence);
 					KeyPropertyResult.bTrackModified = true;
-					
-					GetSequencer()->EmptySelection();
-					GetSequencer()->SelectSection(NewSection);
-					GetSequencer()->ThrobSectionSelection();
+					NewSections.Add(NewSection);
 				}
 			}
 		}
 	}
 
+	if (NewSections.Num() > 0)
+	{
+		GetSequencer()->EmptySelection();
+		for (UMovieSceneSection* NewSection : NewSections)
+		{
+			GetSequencer()->SelectSection(NewSection);
+		}
+		GetSequencer()->ThrobSectionSelection();
+	}
+
 	return KeyPropertyResult;
 }
 
@@ -396,6 +419,12 @@ FKeyPropertyResult FTemplateSequenceTrackEditor::AddLegacyCameraAnimKeyInternal(
 	return FCameraAnimTrackEditorHelper::AddCameraAnimKey(*this, KeyTime, Objects, CameraAnim);
 }
 
+bool FTemplateSequenceTrackEditor::CanAddSubSequence(const UMovieSceneSequence& Sequence) const
+{
+	UMovieSceneSequence* FocusedSequence = GetSequencer()->GetFocusedMovieSceneSequence();
+	return FSubTrackEditorUtil::CanAddSubSequence(FocusedSequence, Sequence);
+}
+
 const UClass* FTemplateSequenceTrackEditor::AcquireObjectClassFromObjectGuid(const FGuid& Guid)
 {
 	if (!Guid.IsValid())
diff --git a/Engine/Plugins/MovieScene/TemplateSequence/Source/TemplateSequenceEditor/Private/TrackEditors/TemplateSequenceTrackEditor.h b/Engine/Plugins/MovieScene/TemplateSequence/Source/TemplateSequenceEditor/Private/TrackEditors/TemplateSequenceTrackEditor.h
--- a/Engine/Plugins/MovieScene/TemplateSequence/Source/TemplateSequenceEditor/Private/TrackEditors/TemplateSequenceTrackEditor.h
+++ b/Engine/Plugins/MovieScene/TemplateSequence/Source/TemplateSequenceEditor/Private/TrackEditors/TemplateSequenceTrackEditor.h
@@ -41,6 +41,7 @@ private:
 	FKeyPropertyResult AddKeyInternal(FFrameNumber KeyTime, TArray<FGuid> ObjectBindings, UTemplateSequence* TemplateSequence);
 	FKeyPropertyResult AddLegacyCameraAnimKeyInternal(FFrameNumber KeyTime, const TArray<TWeakObjectPtr<UObject>> Objects, UCameraAnim* CameraAnim);
 
+	bool CanAddSubSequence(const UMovieSceneSequence& Sequence) const;
 	const UClass* AcquireObjectClassFromObjectGuid(const FGuid& Guid);
 	UCameraComponent* AcquireCameraComponentFromObjectGuid(const FGuid& Guid);
 
diff --git a/Engine/Source/Editor/MovieSceneTools/Private/TrackEditors/SubTrackEditor.cpp b/Engine/Source/Editor/MovieSceneTools/Private/TrackEditors/SubTrackEditor.cpp
--- a/Engine/Source/Editor/MovieSceneTools/Private/TrackEditors/SubTrackEditor.cpp
+++ b/Engine/Source/Editor/MovieSceneTools/Private/TrackEditors/SubTrackEditor.cpp
@@ -480,35 +480,7 @@ bool FSubTrackEditor::CanAddSubSequence(const UMovieSceneSequence& Sequence) con
 {
 	// prevent adding ourselves and ensure we have a valid movie scene
 	UMovieSceneSequence* FocusedSequence = GetSequencer()->GetFocusedMovieSceneSequence();
-
-	if ((FocusedSequence == nullptr) || (FocusedSequence == &Sequence) || (FocusedSequence->GetMovieScene() == nullptr))
-	{
-		return false;
-	}
-
-	// ensure that the other sequence has a valid movie scene
-	UMovieScene* SequenceMovieScene = Sequence.GetMovieScene();
-
-	if (SequenceMovieScene == nullptr)
-	{
-		return false;
-	}
-
-	// make sure we are not contained in the other sequence (circular dependency)
-	// @todo sequencer: this check is not sufficient (does not prevent circular dependencies of 2+ levels)
-	UMovieSceneSubTrack* SequenceSubTrack = SequenceMovieScene->FindMasterTrack<UMovieSceneSubTrack>();
-	if (SequenceSubTrack && SequenceSubTrack->ContainsSequence(*FocusedSequence, true))
-	{
-		return false;
-	}
-
-	UMovieSceneCinematicShotTrack* SequenceCinematicTrack = SequenceMovieScene->FindMasterTrack<UMovieSceneCinematicShotTrack>();
-	if (SequenceCinematicTrack && SequenceCinematicTrack->ContainsSequence(*FocusedSequence, true))
-	{
-		return false;
-	}
-
-	return true;
+	return FSubTrackEditorUtil::CanAddSubSequence(FocusedSequence, Sequence);
 }
 
 
diff --git a/Engine/Source/Editor/MovieSceneTools/Private/TrackEditors/SubTrackEditorBase.cpp b/Engine/Source/Editor/MovieSceneTools/Private/TrackEditors/SubTrackEditorBase.cpp
--- a/Engine/Source/Editor/MovieSceneTools/Private/TrackEditors/SubTrackEditorBase.cpp
+++ b/Engine/Source/Editor/MovieSceneTools/Private/TrackEditors/SubTrackEditorBase.cpp
@@ -1,6 +1,8 @@
 // Copyright Epic Games, Inc. All Rights Reserved.
 
 #include "TrackEditors/SubTrackEditorBase.h"
+#include "Tracks/MovieSceneCinematicShotTrack.h"
+#include "Tracks/MovieSceneSubTrack.h"
 
 #define LOCTEXT_NAMESPACE "FSubTrackEditorBase"
 
@@ -344,4 +346,37 @@ FFrameNumber FSubSectionEditorUtil::SlipSection(FFrameNumber SlipTime)
     return SlipTime;
 }
 
+bool FSubTrackEditorUtil::CanAddSubSequence(const UMovieSceneSequence* CurrentSequence, const UMovieSceneSequence& SubSequence)
+{
+	// Prevent adding ourselves and ensure we have a valid movie scene.
+	if ((CurrentSequence == nullptr) || (CurrentSequence == &SubSequence) || (CurrentSequence->GetMovieScene() == nullptr))
+	{
+		return false;
+	}
+
+	// ensure that the other sequence has a valid movie scene
+	UMovieScene* SequenceMovieScene = SubSequence.GetMovieScene();
+
+	if (SequenceMovieScene == nullptr)
+	{
+		return false;
+	}
+
+	// make sure we are not contained in the other sequence (circular dependency)
+	// @todo sequencer: this check is not sufficient (does not prevent circular dependencies of 2+ levels)
+	UMovieSceneSubTrack* SequenceSubTrack = SequenceMovieScene->FindMasterTrack<UMovieSceneSubTrack>();
+	if (SequenceSubTrack && SequenceSubTrack->ContainsSequence(*CurrentSequence, true))
+	{
+		return false;
+	}
+
+	UMovieSceneCinematicShotTrack* SequenceCinematicTrack = SequenceMovieScene->FindMasterTrack<UMovieSceneCinematicShotTrack>();
+	if (SequenceCinematicTrack && SequenceCinematicTrack->ContainsSequence(*CurrentSequence, true))
+	{
+		return false;
+	}
+
+	return true;
+}
+
 #undef LOCTEXT_NAMESPACE
diff --git a/Engine/Source/Editor/MovieSceneTools/Public/TrackEditors/SubTrackEditorBase.h b/Engine/Source/Editor/MovieSceneTools/Public/TrackEditors/SubTrackEditorBase.h
--- a/Engine/Source/Editor/MovieSceneTools/Public/TrackEditors/SubTrackEditorBase.h
+++ b/Engine/Source/Editor/MovieSceneTools/Public/TrackEditors/SubTrackEditorBase.h
@@ -88,6 +88,21 @@ private:
     FFrameNumber InitialStartTimeDuringResize;
 };
 
+class MOVIESCENETOOLS_API FSubTrackEditorUtil
+{
+public:
+	/**
+	 * Check whether the given sequence can be added as a sub-sequence.
+	 *
+	 * The purpose of this method is to disallow circular references
+	 * between sub-sequences in the focused movie scene.
+	 *
+	 * @param Sequence The sequence to check.
+	 * @return true if the sequence can be added as a sub-sequence, false otherwise.
+	 */
+	static bool CanAddSubSequence(const UMovieSceneSequence* CurrentSequence, const UMovieSceneSequence& SubSequence);
+};
+
 /**
  * Mixin class for sub-sequence section interfaces.
  *
