Fixed non-unity build errors. CL from David.Bollo

diff --git a/Engine/Plugins/Lumin/MagicLeap/Source/MagicLeapAR/Private/LuminAROrigin.cpp b/Engine/Plugins/Lumin/MagicLeap/Source/MagicLeapAR/Private/LuminAROrigin.cpp
--- a/Engine/Plugins/Lumin/MagicLeap/Source/MagicLeapAR/Private/LuminAROrigin.cpp
+++ b/Engine/Plugins/Lumin/MagicLeap/Source/MagicLeapAR/Private/LuminAROrigin.cpp
@@ -5,6 +5,7 @@
 #include "UObject/ConstructorHelpers.h"
 #include "AugmentedReality/Public/ARSessionConfig.h"
 #include "HeadMountedDisplayFunctionLibrary.h"
+#include "Materials/MaterialInterface.h"
 
 DEFINE_LOG_CATEGORY_STATIC(LogMagicLeap, Log, All);
 
diff --git a/Engine/Plugins/MovieScene/TemplateSequence/Source/TemplateSequence/Private/SequenceCameraShake.cpp b/Engine/Plugins/MovieScene/TemplateSequence/Source/TemplateSequence/Private/SequenceCameraShake.cpp
--- a/Engine/Plugins/MovieScene/TemplateSequence/Source/TemplateSequence/Private/SequenceCameraShake.cpp
+++ b/Engine/Plugins/MovieScene/TemplateSequence/Source/TemplateSequence/Private/SequenceCameraShake.cpp
@@ -7,6 +7,7 @@
 #include "Camera/PlayerCameraManager.h"
 #include "CineCameraActor.h"
 #include "Containers/ArrayView.h"
+#include "Engine/World.h"
 #include "EntitySystem/MovieSceneBoundSceneComponentInstantiator.h"
 #include "EntitySystem/MovieSceneEntitySystem.h"
 #include "EntitySystem/MovieScenePropertySystemTypes.h"
diff --git a/Engine/Source/Editor/UnrealEd/Classes/Commandlets/WorldPartitionConvertCommandlet.h b/Engine/Source/Editor/UnrealEd/Classes/Commandlets/WorldPartitionConvertCommandlet.h
--- a/Engine/Source/Editor/UnrealEd/Classes/Commandlets/WorldPartitionConvertCommandlet.h
+++ b/Engine/Source/Editor/UnrealEd/Classes/Commandlets/WorldPartitionConvertCommandlet.h
@@ -14,6 +14,7 @@
 
 UNREALED_API DECLARE_LOG_CATEGORY_EXTERN(LogWorldPartitionConvertCommandlet, Log, All);
 
+class ULevelStreaming;
 class UWorldPartition;
 class UWorldComposition;
 
diff --git a/Engine/Source/Runtime/Engine/Classes/PhysicsEngine/PhysicsSettings.h b/Engine/Source/Runtime/Engine/Classes/PhysicsEngine/PhysicsSettings.h
--- a/Engine/Source/Runtime/Engine/Classes/PhysicsEngine/PhysicsSettings.h
+++ b/Engine/Source/Runtime/Engine/Classes/PhysicsEngine/PhysicsSettings.h
@@ -13,6 +13,7 @@
 #include "Engine/DeveloperSettings.h"
 #include "PhysicsSettingsEnums.h"
 #include "BodySetupEnums.h"
+#include "Chaos/ChaosEngineInterface.h"
 #include "GameFramework/WorldSettings.h"
 #include "PhysicsCoreTypes.h"
 #include "PhysicsSettingsCore.h"
diff --git a/Engine/Source/Runtime/Engine/Public/WorldPartition/HLOD/HLODLayer.h b/Engine/Source/Runtime/Engine/Public/WorldPartition/HLOD/HLODLayer.h
--- a/Engine/Source/Runtime/Engine/Public/WorldPartition/HLOD/HLODLayer.h
+++ b/Engine/Source/Runtime/Engine/Public/WorldPartition/HLOD/HLODLayer.h
@@ -3,6 +3,7 @@
 #pragma once
 
 #include "CoreMinimal.h"
+#include "EngineDefines.h"
 #include "UObject/ObjectMacros.h"
 #include "UObject/UObjectGlobals.h"
 #include "UObject/Object.h"
@@ -14,6 +15,7 @@
 
 #include "HLODLayer.generated.h"
 
+class UMaterial;
 class UWorldPartition;
 class AActor;
 
