Template Sequence: Supports level sequence only (not UMG)

diff --git a/Engine/Plugins/MovieScene/TemplateSequence/Source/TemplateSequenceEditor/Private/TrackEditors/TemplateSequenceTrackEditor.cpp b/Engine/Plugins/MovieScene/TemplateSequence/Source/TemplateSequenceEditor/Private/TrackEditors/TemplateSequenceTrackEditor.cpp
--- a/Engine/Plugins/MovieScene/TemplateSequence/Source/TemplateSequenceEditor/Private/TrackEditors/TemplateSequenceTrackEditor.cpp
+++ b/Engine/Plugins/MovieScene/TemplateSequence/Source/TemplateSequenceEditor/Private/TrackEditors/TemplateSequenceTrackEditor.cpp
@@ -15,6 +15,7 @@
 #include "Framework/Notifications/NotificationManager.h"
 #include "ICollectionManager.h"
 #include "IContentBrowserSingleton.h"
+#include "LevelSequence.h"
 #include "Misc/ConfigCacheIni.h"
 #include "MovieSceneCommonHelpers.h"
 #include "MovieSceneTimeHelpers.h"
@@ -50,6 +51,11 @@ bool FTemplateSequenceTrackEditor::SupportsType(TSubclassOf<UMovieSceneTrack> Ty
 	return Type == UTemplateSequenceTrack::StaticClass();
 }
 
+bool FTemplateSequenceTrackEditor::SupportsSequence(UMovieSceneSequence* InSequence) const
+{
+	return InSequence && InSequence->IsA(ULevelSequence::StaticClass());
+}
+
 void FTemplateSequenceTrackEditor::BuildObjectBindingTrackMenu(FMenuBuilder& MenuBuilder, const TArray<FGuid>& ObjectBindings, const UClass* ObjectClass)
 {
 	bool bIsCameraAnimMenu = true;
diff --git a/Engine/Plugins/MovieScene/TemplateSequence/Source/TemplateSequenceEditor/Private/TrackEditors/TemplateSequenceTrackEditor.h b/Engine/Plugins/MovieScene/TemplateSequence/Source/TemplateSequenceEditor/Private/TrackEditors/TemplateSequenceTrackEditor.h
--- a/Engine/Plugins/MovieScene/TemplateSequence/Source/TemplateSequenceEditor/Private/TrackEditors/TemplateSequenceTrackEditor.h
+++ b/Engine/Plugins/MovieScene/TemplateSequence/Source/TemplateSequenceEditor/Private/TrackEditors/TemplateSequenceTrackEditor.h
@@ -25,6 +25,7 @@ public:
 public:
 	// ISequencerTrackEditor interface
 	virtual bool SupportsType(TSubclassOf<class UMovieSceneTrack> TrackClass) const override;
+	virtual bool SupportsSequence(UMovieSceneSequence* InSequence) const override;
 	virtual void BuildObjectBindingTrackMenu(FMenuBuilder& MenuBuilder, const TArray<FGuid>& ObjectBindings, const UClass* ObjectClass) override;
 	TSharedPtr<SWidget> BuildOutlinerEditWidget(const FGuid& ObjectBinding, UMovieSceneTrack* Track, const FBuildEditWidgetParams& Params) override;
 	virtual TSharedRef<ISequencerSection> MakeSectionInterface(UMovieSceneSection& SectionObject, UMovieSceneTrack& Track, FGuid ObjectBinding) override;
