Sequencer: Implement scrubbing support for Sequence camera shakes.

diff --git a/Engine/Plugins/MovieScene/TemplateSequence/Source/TemplateSequence/Private/SequenceCameraShake.cpp b/Engine/Plugins/MovieScene/TemplateSequence/Source/TemplateSequence/Private/SequenceCameraShake.cpp
--- a/Engine/Plugins/MovieScene/TemplateSequence/Source/TemplateSequence/Private/SequenceCameraShake.cpp
+++ b/Engine/Plugins/MovieScene/TemplateSequence/Source/TemplateSequence/Private/SequenceCameraShake.cpp
@@ -229,45 +229,20 @@
 
 void USequenceCameraShakePattern::UpdateShakePatternImpl(const FCameraShakeUpdateParams& Params, FCameraShakeUpdateResult& OutResult)
 {
-	using namespace UE::MovieScene;
-
-	check(CameraStandIn);
-
-	// Reset the camera stand-in's properties based on the new "current" (unshaken) values.
-	CameraStandIn->Reset(Params.POV);
-
-	// Sequencer animates things based on the initial values cached when the sequence started. But here we want
-	// to animate things based on the moving current values of the camera... i.e., we want to shake a constantly
-	// moving camera. So every frame, we need to update the initial values that sequencer uses.
-	UpdateInitialCameraStandInPropertyValues();
-
-	// Get the "unshaken" properties that need to be treated additively.
-	const float OriginalFieldOfView = CameraStandIn->FieldOfView;
-
-	// Update the sequence.
 	const FFrameRate TickResolution = Player->GetInputRate();
 	const FFrameTime NewPosition = Player->GetCurrentPosition() + Params.DeltaTime * PlayRate * TickResolution;
-	Player->Update(NewPosition);
-
-	// Recalculate properties that might be invalidated by other properties having been animated.
-	CameraStandIn->RecalcDerivedData();
-
-	// Grab the final animated (shaken) values, figure out the delta, apply scale, and feed that into the 
-	// camera shake result.
-	// Transform is always treated as a local, additive value. The data better be good.
-	const FTransform ShakenTransform = CameraStandIn->GetTransform();
-	OutResult.Location = ShakenTransform.GetLocation() * Scale;
-	OutResult.Rotation = ShakenTransform.GetRotation().Rotator() * Scale;
-
-	// FieldOfView follows the current camera's value every frame, so we can compute how much the shake is
-	// changing it.
-	const float ShakenFieldOfView = CameraStandIn->FieldOfView;
-	const float DeltaFieldOfView = ShakenFieldOfView - OriginalFieldOfView;
-	OutResult.FOV = DeltaFieldOfView * Scale;
-
-	// The other properties aren't treated as additive.
-	OutResult.PostProcessSettings = CameraStandIn->PostProcessSettings;
-	OutResult.PostProcessBlendWeight = CameraStandIn->PostProcessBlendWeight;
+	UpdateCamera(NewPosition, Params.POV, OutResult);
 }
 
+void USequenceCameraShakePattern::ScrubShakePatternImpl(const FCameraShakeScrubParams& Params, FCameraShakeUpdateResult& OutResult)
+{
+	Player->StartScrubbing();
+	{
+		const FFrameRate TickResolution = Player->GetInputRate();
+		const FFrameTime NewPosition = Params.AbsoluteTime * PlayRate * TickResolution;
+		UpdateCamera(NewPosition, Params.POV, OutResult);
+	}
+	Player->EndScrubbing();
+}
+
 void USequenceCameraShakePattern::StopShakePatternImpl(const FCameraShakeStopParams& Params)
@@ -323,6 +298,47 @@
 	}
 }
 
+void USequenceCameraShakePattern::UpdateCamera(FFrameTime NewPosition, const FMinimalViewInfo& InPOV, FCameraShakeUpdateResult& OutResult)
+{
+	using namespace UE::MovieScene;
+
+	check(CameraStandIn);
+
+	// Reset the camera stand-in's properties based on the new "current" (unshaken) values.
+	CameraStandIn->Reset(InPOV);
+
+	// Sequencer animates things based on the initial values cached when the sequence started. But here we want
+	// to animate things based on the moving current values of the camera... i.e., we want to shake a constantly
+	// moving camera. So every frame, we need to update the initial values that sequencer uses.
+	UpdateInitialCameraStandInPropertyValues();
+
+	// Get the "unshaken" properties that need to be treated additively.
+	const float OriginalFieldOfView = CameraStandIn->FieldOfView;
+
+	// Update the sequence.
+	Player->Update(NewPosition);
+
+	// Recalculate properties that might be invalidated by other properties having been animated.
+	CameraStandIn->RecalcDerivedData();
+
+	// Grab the final animated (shaken) values, figure out the delta, apply scale, and feed that into the 
+	// camera shake result.
+	// Transform is always treated as a local, additive value. The data better be good.
+	const FTransform ShakenTransform = CameraStandIn->GetTransform();
+	OutResult.Location = ShakenTransform.GetLocation() * Scale;
+	OutResult.Rotation = ShakenTransform.GetRotation().Rotator() * Scale;
+
+	// FieldOfView follows the current camera's value every frame, so we can compute how much the shake is
+	// changing it.
+	const float ShakenFieldOfView = CameraStandIn->FieldOfView;
+	const float DeltaFieldOfView = ShakenFieldOfView - OriginalFieldOfView;
+	OutResult.FOV = DeltaFieldOfView * Scale;
+
+	// The other properties aren't treated as additive.
+	OutResult.PostProcessSettings = CameraStandIn->PostProcessSettings;
+	OutResult.PostProcessBlendWeight = CameraStandIn->PostProcessBlendWeight;
+}
+
 void USequenceCameraShakePattern::UpdateInitialCameraStandInPropertyValues()
 {
 	using namespace UE::MovieScene;
@@ -448,7 +464,7 @@
 
 void USequenceCameraShakeSequencePlayer::Update(FFrameTime NewPosition)
 {
-	check(Status == EMovieScenePlayerStatus::Playing);
+	check(Status == EMovieScenePlayerStatus::Playing || Status == EMovieScenePlayerStatus::Scrubbing);
 	check(RootTemplateInstance.IsValid());
 
 	if (bIsLooping)
@@ -488,3 +504,16 @@
 		RootTemplateInstance.Finish(*this);
 	}
 }
+
+void USequenceCameraShakeSequencePlayer::StartScrubbing()
+{
+	ensure(Status == EMovieScenePlayerStatus::Playing);
+	Status = EMovieScenePlayerStatus::Scrubbing;
+}
+
+void USequenceCameraShakeSequencePlayer::EndScrubbing()
+{
+	ensure(Status == EMovieScenePlayerStatus::Scrubbing);
+	Status = EMovieScenePlayerStatus::Playing;
+}
+
diff --git a/Engine/Plugins/MovieScene/TemplateSequence/Source/TemplateSequence/Public/SequenceCameraShake.h b/Engine/Plugins/MovieScene/TemplateSequence/Source/TemplateSequence/Public/SequenceCameraShake.h
--- a/Engine/Plugins/MovieScene/TemplateSequence/Source/TemplateSequence/Public/SequenceCameraShake.h
+++ b/Engine/Plugins/MovieScene/TemplateSequence/Source/TemplateSequence/Public/SequenceCameraShake.h
@@ -141,11 +141,13 @@ private:
 	virtual void GetShakePatternInfoImpl(FCameraShakeInfo& OutInfo) const override;
 	virtual void StartShakePatternImpl(const FCameraShakeStartParams& Params) override;
 	virtual void UpdateShakePatternImpl(const FCameraShakeUpdateParams& Params, FCameraShakeUpdateResult& OutResult) override;
+	virtual void ScrubShakePatternImpl(const FCameraShakeScrubParams& Params, FCameraShakeUpdateResult& OutResult) override;
 	virtual void StopShakePatternImpl(const FCameraShakeStopParams& Params) override;
 	virtual void TeardownShakePatternImpl() override;
 
 	static void RegisterCameraStandIn();
 
+	void UpdateCamera(FFrameTime NewPosition, const FMinimalViewInfo& InPOV, FCameraShakeUpdateResult& OutResult);
 	void UpdateInitialCameraStandInPropertyValues();
 
 private:
@@ -203,6 +205,11 @@ public:
 	/** Stop playing the sequence */
 	void Stop();
 
+	/** Sets the player in scrub mode */
+	void StartScrubbing();
+	/** Ends scrub mode */
+	void EndScrubbing();
+
 	/** Gets the current play position */
 	FFrameTime GetCurrentPosition() const { return PlayPosition.GetCurrentPosition(); }
 	/** Get the sequence tick resolution */
