Template Sequence: Return if a spawnable could not be created (ie. no spawner for the actor type). This fixes a crash when trying to use an invalid type. Note there is already an error popup, but it crashed after that dereferencing the null spawnable.

diff --git a/Engine/Plugins/MovieScene/TemplateSequence/Source/TemplateSequenceEditor/Private/Misc/TemplateSequenceEditorUtil.cpp b/Engine/Plugins/MovieScene/TemplateSequence/Source/TemplateSequenceEditor/Private/Misc/TemplateSequenceEditorUtil.cpp
--- a/Engine/Plugins/MovieScene/TemplateSequence/Source/TemplateSequenceEditor/Private/Misc/TemplateSequenceEditorUtil.cpp
+++ b/Engine/Plugins/MovieScene/TemplateSequence/Source/TemplateSequenceEditor/Private/Misc/TemplateSequenceEditorUtil.cpp
@@ -44,6 +44,10 @@ void FTemplateSequenceEditorUtil::ChangeActorBinding(UObject* Object, UActorFact
 	// Make the new spawnable object binding.
 	FGuid NewSpawnableGuid = Sequencer.MakeNewSpawnable(*Object);
 	FMovieSceneSpawnable* NewSpawnable = MovieScene->FindSpawnable(NewSpawnableGuid);
+	if (!NewSpawnable)
+	{
+		return;
+	}
 
 	// Pre-notify that a change will occur.
 	FProperty* BoundActorClassProperty = FindFProperty<FProperty>(UTemplateSequence::StaticClass(), GET_MEMBER_NAME_CHECKED(UTemplateSequence, BoundActorClass));
