Template Sequence: TemplateSequenceTrackEditor should also support template sequences, which therefore allow a template sequence to act on another template sequence. This fixes a crash when trying to open up a template sequence which refers to another template sequence.

diff --git a/Engine/Plugins/MovieScene/TemplateSequence/Source/TemplateSequenceEditor/Private/TrackEditors/TemplateSequenceTrackEditor.cpp b/Engine/Plugins/MovieScene/TemplateSequence/Source/TemplateSequenceEditor/Private/TrackEditors/TemplateSequenceTrackEditor.cpp
--- a/Engine/Plugins/MovieScene/TemplateSequence/Source/TemplateSequenceEditor/Private/TrackEditors/TemplateSequenceTrackEditor.cpp
+++ b/Engine/Plugins/MovieScene/TemplateSequence/Source/TemplateSequenceEditor/Private/TrackEditors/TemplateSequenceTrackEditor.cpp
@@ -53,7 +53,7 @@ bool FTemplateSequenceTrackEditor::SupportsType(TSubclassOf<UMovieSceneTrack> Ty
 
 bool FTemplateSequenceTrackEditor::SupportsSequence(UMovieSceneSequence* InSequence) const
 {
-	return InSequence && InSequence->IsA(ULevelSequence::StaticClass());
+	return InSequence && (InSequence->IsA(ULevelSequence::StaticClass()) || InSequence->IsA(UTemplateSequence::StaticClass()));
 }
 
 void FTemplateSequenceTrackEditor::BuildObjectBindingTrackMenu(FMenuBuilder& MenuBuilder, const TArray<FGuid>& ObjectBindings, const UClass* ObjectClass)
