Template Sequence: Set the Sequencer tab to be a null widget just like in FLevelSequenceEditorToolkit when closed. This fixes a crash in the editor when you delete the currently opened template sequence.

diff --git a/Engine/Plugins/MovieScene/TemplateSequence/Source/TemplateSequenceEditor/Private/TemplateSequenceEditorToolkit.cpp b/Engine/Plugins/MovieScene/TemplateSequence/Source/TemplateSequenceEditor/Private/TemplateSequenceEditorToolkit.cpp
--- a/Engine/Plugins/MovieScene/TemplateSequence/Source/TemplateSequenceEditor/Private/TemplateSequenceEditorToolkit.cpp
+++ b/Engine/Plugins/MovieScene/TemplateSequence/Source/TemplateSequenceEditor/Private/TemplateSequenceEditorToolkit.cpp
@@ -50,7 +50,12 @@ FTemplateSequenceEditorToolkit::~FTemplateSequenceEditorToolkit()
 
 	if (FModuleManager::Get().IsModuleLoaded(TEXT("LevelEditor")))
 	{
-		auto& LevelEditorModule = FModuleManager::LoadModuleChecked<FLevelEditorModule>(TEXT("LevelEditor"));
+		FLevelEditorModule& LevelEditorModule = FModuleManager::LoadModuleChecked<FLevelEditorModule>("LevelEditor");
+		
+		// @todo remove when world-centric mode is added
+		LevelEditorModule.AttachSequencer(SNullWidget::NullWidget, nullptr);
+		FLevelEditorSequencerIntegration::Get().RemoveSequencer(Sequencer.ToSharedRef());
+
 		LevelEditorModule.OnMapChanged().RemoveAll(this);
 	}
 
