Template Sequence: Disallow converting a spawnable to a possessable for template sequences. This fixes a crash when trying to convert the bound object in a template sequence to a possessable.

diff --git a/Engine/Plugins/MovieScene/TemplateSequence/Source/TemplateSequenceEditor/Private/Misc/TemplateSequenceEditorSpawnRegister.cpp b/Engine/Plugins/MovieScene/TemplateSequence/Source/TemplateSequenceEditor/Private/Misc/TemplateSequenceEditorSpawnRegister.cpp
--- a/Engine/Plugins/MovieScene/TemplateSequence/Source/TemplateSequenceEditor/Private/Misc/TemplateSequenceEditorSpawnRegister.cpp
+++ b/Engine/Plugins/MovieScene/TemplateSequence/Source/TemplateSequenceEditor/Private/Misc/TemplateSequenceEditorSpawnRegister.cpp
@@ -1,3 +1,12 @@
 // Copyright Epic Games, Inc. All Rights Reserved.
 
 #include "TemplateSequenceEditorSpawnRegister.h"
+
+#if WITH_EDITOR
+
+bool FTemplateSequenceEditorSpawnRegister::CanConvertSpawnableToPossessable(FMovieSceneSpawnable& Spawnable) const
+{
+	return false;
+}
+
+#endif
\ No newline at end of file
diff --git a/Engine/Plugins/MovieScene/TemplateSequence/Source/TemplateSequenceEditor/Private/Misc/TemplateSequenceEditorSpawnRegister.h b/Engine/Plugins/MovieScene/TemplateSequence/Source/TemplateSequenceEditor/Private/Misc/TemplateSequenceEditorSpawnRegister.h
--- a/Engine/Plugins/MovieScene/TemplateSequence/Source/TemplateSequenceEditor/Private/Misc/TemplateSequenceEditorSpawnRegister.h
+++ b/Engine/Plugins/MovieScene/TemplateSequence/Source/TemplateSequenceEditor/Private/Misc/TemplateSequenceEditorSpawnRegister.h
@@ -2,9 +2,13 @@
 
 #pragma once
 
+#include "MovieSceneSpawnable.h"
 #include "Misc/LevelSequenceEditorSpawnRegister.h"
 
 /** Movie scene spawn register that knows how to handle spawning objects (actors) for a template sequence  */
 class FTemplateSequenceEditorSpawnRegister : public FLevelSequenceEditorSpawnRegister
 {
+#if WITH_EDITOR
+	virtual bool CanConvertSpawnableToPossessable(FMovieSceneSpawnable& Spawnable) const override;
+#endif
 };
