Sequencer: Fix dependency problem with template sequence property multipliers.

diff --git a/Engine/Plugins/MovieScene/TemplateSequence/Source/TemplateSequence/Private/Systems/TemplateSequenceSystem.cpp b/Engine/Plugins/MovieScene/TemplateSequence/Source/TemplateSequence/Private/Systems/TemplateSequenceSystem.cpp
--- a/Engine/Plugins/MovieScene/TemplateSequence/Source/TemplateSequence/Private/Systems/TemplateSequenceSystem.cpp
+++ b/Engine/Plugins/MovieScene/TemplateSequence/Source/TemplateSequence/Private/Systems/TemplateSequenceSystem.cpp
@@ -13,11 +13,14 @@
 #include "IMovieScenePlayer.h"
 #include "MovieSceneTracksComponentTypes.h"
 #include "Sections/TemplateSequenceSection.h"
-#include "Systems/FloatChannelEvaluatorSystem.h"
 #include "Systems/DoubleChannelEvaluatorSystem.h"
+#include "Systems/FloatChannelEvaluatorSystem.h"
+#include "Systems/MovieScene3DTransformPropertySystem.h"
 #include "Systems/MovieSceneBaseValueEvaluatorSystem.h"
-#include "Systems/MovieScenePiecewiseFloatBlenderSystem.h"
+#include "Systems/MovieSceneComponentTransformSystem.h"
+#include "Systems/MovieSceneFloatPropertySystem.h"
 #include "Systems/MovieScenePiecewiseDoubleBlenderSystem.h"
+#include "Systems/MovieScenePiecewiseFloatBlenderSystem.h"
 #include "TemplateSequence.h"
 #include "TemplateSequenceComponentTypes.h"
 
@@ -270,6 +273,11 @@ UTemplateSequencePropertyScalingEvaluatorSystem::UTemplateSequencePropertyScalin
 		// We need to multiply values before they are blended together with other values we shouldn't touch.
 		DefineImplicitPrerequisite(GetClass(), UMovieScenePiecewiseFloatBlenderSystem::StaticClass());
 		DefineImplicitPrerequisite(GetClass(), UMovieScenePiecewiseDoubleBlenderSystem::StaticClass());
+		// We need to multiply values before they are set on the animated objects. For blended values, the dependencies above are
+		// enough, but for non-blended values that use the "fast-path", we need this dependency.
+		DefineImplicitPrerequisite(GetClass(), UMovieSceneFloatPropertySystem::StaticClass());
+		DefineImplicitPrerequisite(GetClass(), UMovieSceneComponentTransformSystem::StaticClass());
+		DefineImplicitPrerequisite(GetClass(), UMovieScene3DTransformPropertySystem::StaticClass());
 		// We need this component to lookup the binding GUID for an entity, so we know what scaling to apply to it.
 		DefineComponentConsumer(GetClass(), BuiltInComponents->BoundObject);
 	}
