Template Sequence: Fix import fbx onto template sequence

Import fbx onto template sequence works by creating a temporary template sequence player, spawns the necessary objects, and then matches the objects to the imported fbx data. The issue here was that the spawnables weren't spawning because the binding override is by default to not override the default. With template sequences, usually you have another object that takes the place of the binding. By setting this false in FindOrAddBinding, this allows the temporarily spawned objects to be found and mapped onto.

diff --git a/Engine/Plugins/MovieScene/TemplateSequence/Source/TemplateSequence/Private/TemplateSequence.cpp b/Engine/Plugins/MovieScene/TemplateSequence/Source/TemplateSequence/Private/TemplateSequence.cpp
--- a/Engine/Plugins/MovieScene/TemplateSequence/Source/TemplateSequence/Private/TemplateSequence.cpp
+++ b/Engine/Plugins/MovieScene/TemplateSequence/Source/TemplateSequence/Private/TemplateSequence.cpp
@@ -140,12 +140,15 @@ FGuid UTemplateSequence::FindOrAddBinding(UObject* InObject)
 	}
 
 	// Perform a potentially slow lookup of every possessable binding in the sequence to see if we already have this
-	UWorld* World = GetWorld();
+	UWorld* World = InObject->GetWorld();
 	ensure(World);
 	{
 		ATemplateSequenceActor* OutActor = nullptr;
 		UTemplateSequencePlayer* Player = UTemplateSequencePlayer::CreateTemplateSequencePlayer(World, this, FMovieSceneSequencePlaybackSettings(), OutActor);
 		ensure(Player);
+		
+		// Don't override the default, the default assignment should remain bound so that objects will spawn
+		OutActor->BindingOverride.bOverridesDefault = false;
 
 		Player->Initialize(this, World, FMovieSceneSequencePlaybackSettings());
 		Player->State.AssignSequence(MovieSceneSequenceID::Root, *this, *Player);
diff --git a/Engine/Source/Editor/Sequencer/Private/Sequencer.cpp b/Engine/Source/Editor/Sequencer/Private/Sequencer.cpp
--- a/Engine/Source/Editor/Sequencer/Private/Sequencer.cpp
+++ b/Engine/Source/Editor/Sequencer/Private/Sequencer.cpp
@@ -11837,7 +11837,11 @@ void FSequencer::NewCameraAdded(ACameraActor* NewCamera, FGuid CameraGuid)
 
 	MovieSceneToolHelpers::LockCameraActorToViewport(SharedThis(this), NewCamera);
 
-	MovieSceneToolHelpers::CreateCameraCutSectionForCamera(GetFocusedMovieSceneSequence()->GetMovieScene(), CameraGuid, GetLocalTime().Time.FloorToFrame());
+	UMovieSceneSequence* Sequence = GetFocusedMovieSceneSequence();
+	if (Sequence->IsTrackSupported(UMovieSceneCameraCutTrack::StaticClass()) == ETrackSupport::Supported)
+	{
+		MovieSceneToolHelpers::CreateCameraCutSectionForCamera(Sequence->GetMovieScene(), CameraGuid, GetLocalTime().Time.FloorToFrame());
+	}
 }
 
 
