Correctly initialize camera animation sequences upon creation

It was previously only initialized when the asset was opened in Sequencer for the first time. It could lead to problems if the asset was added to something else (like a camera shake or a level sequence) before being edited.

diff --git a/Engine/Plugins/MovieScene/TemplateSequence/Source/TemplateSequenceEditor/Private/AssetTools/CameraAnimationSequenceActions.cpp b/Engine/Plugins/MovieScene/TemplateSequence/Source/TemplateSequenceEditor/Private/AssetTools/CameraAnimationSequenceActions.cpp
--- a/Engine/Plugins/MovieScene/TemplateSequence/Source/TemplateSequenceEditor/Private/AssetTools/CameraAnimationSequenceActions.cpp
+++ b/Engine/Plugins/MovieScene/TemplateSequence/Source/TemplateSequenceEditor/Private/AssetTools/CameraAnimationSequenceActions.cpp
@@ -27,7 +27,6 @@ void FCameraAnimationSequenceActions::InitializeToolkitParams(FTemplateSequenceT
     FTemplateSequenceActions::InitializeToolkitParams(ToolkitParams);
 
     ToolkitParams.bCanChangeBinding = false;
-    ToolkitParams.InitialBindingClass = ACineCameraActor::StaticClass();
 }
 
 #undef LOCTEXT_NAMESPACE
diff --git a/Engine/Plugins/MovieScene/TemplateSequence/Source/TemplateSequenceEditor/Private/Factories/CameraAnimationSequenceFactoryNew.cpp b/Engine/Plugins/MovieScene/TemplateSequence/Source/TemplateSequenceEditor/Private/Factories/CameraAnimationSequenceFactoryNew.cpp
--- a/Engine/Plugins/MovieScene/TemplateSequence/Source/TemplateSequenceEditor/Private/Factories/CameraAnimationSequenceFactoryNew.cpp
+++ b/Engine/Plugins/MovieScene/TemplateSequence/Source/TemplateSequenceEditor/Private/Factories/CameraAnimationSequenceFactoryNew.cpp
@@ -2,10 +2,11 @@
 
 #include "Factories/CameraAnimationSequenceFactoryNew.h"
 #include "CameraAnimationSequence.h"
+#include "CineCameraActor.h"
+#include "Factories/TemplateSequenceFactoryUtil.h"
 #include "MovieScene.h"
 #include "MovieSceneToolsProjectSettings.h"
-
-#define LOCTEXT_NAMESPACE "MovieSceneFactory"
+#include "Misc/TemplateSequenceEditorUtil.h"
 
 UCameraAnimationSequenceFactoryNew::UCameraAnimationSequenceFactoryNew(const FObjectInitializer& ObjectInitializer)
 	: Super(ObjectInitializer)
@@ -17,22 +18,10 @@ UCameraAnimationSequenceFactoryNew::UCameraAnimationSequenceFactoryNew(const FOb
 
 UObject* UCameraAnimationSequenceFactoryNew::FactoryCreateNew(UClass* Class, UObject* InParent, FName Name, EObjectFlags Flags, UObject* Context, FFeedbackContext* Warn)
 {
-	auto NewCameraAnimationSequence = NewObject<UCameraAnimationSequence>(InParent, Name, Flags | RF_Transactional);
-	NewCameraAnimationSequence->Initialize();
-
-	// Set up some sensible defaults
-	const UMovieSceneToolsProjectSettings* ProjectSettings = GetDefault<UMovieSceneToolsProjectSettings>();
-
-	FFrameRate TickResolution = NewCameraAnimationSequence->GetMovieScene()->GetTickResolution();
-	NewCameraAnimationSequence->GetMovieScene()->SetPlaybackRange((ProjectSettings->DefaultStartTime*TickResolution).FloorToFrame(), (ProjectSettings->DefaultDuration*TickResolution).FloorToFrame().Value);
-
-	return NewCameraAnimationSequence;
+	return FTemplateSequenceFactoryUtil::CreateTemplateSequence(InParent, Name, Flags, UCameraAnimationSequence::StaticClass(), ACineCameraActor::StaticClass());
 }
 
 bool UCameraAnimationSequenceFactoryNew::ShouldShowInNewMenu() const
 {
 	return true;
 }
-
-#undef LOCTEXT_NAMESPACE
-
diff --git a/Engine/Plugins/MovieScene/TemplateSequence/Source/TemplateSequenceEditor/Private/Factories/TemplateSequenceFactoryNew.cpp b/Engine/Plugins/MovieScene/TemplateSequence/Source/TemplateSequenceEditor/Private/Factories/TemplateSequenceFactoryNew.cpp
--- a/Engine/Plugins/MovieScene/TemplateSequence/Source/TemplateSequenceEditor/Private/Factories/TemplateSequenceFactoryNew.cpp
+++ b/Engine/Plugins/MovieScene/TemplateSequence/Source/TemplateSequenceEditor/Private/Factories/TemplateSequenceFactoryNew.cpp
@@ -1,20 +1,13 @@
 // Copyright Epic Games, Inc. All Rights Reserved.
 
 #include "Factories/TemplateSequenceFactoryNew.h"
-#include "Framework/Notifications/NotificationManager.h"
+#include "Factories/TemplateSequenceFactoryUtil.h"
 #include "GameFramework/Actor.h"
-#include "ILevelSequenceModule.h"
 #include "Kismet2/SClassPickerDialog.h"
-#include "Misc/TemplateSequenceEditorUtil.h"
-#include "Modules/ModuleManager.h"
 #include "MovieScene.h"
-#include "MovieSceneToolsProjectSettings.h"
-#include "Sections/MovieSceneBoolSection.h"
 #include "TemplateSequence.h"
-#include "Tracks/MovieSceneSpawnTrack.h"
-#include "Widgets/Notifications/SNotificationList.h"
 
-#define LOCTEXT_NAMESPACE "MovieSceneFactory"
+#define LOCTEXT_NAMESPACE "TemplateSequenceFactory"
 
 UTemplateSequenceFactoryNew::UTemplateSequenceFactoryNew(const FObjectInitializer& ObjectInitializer)
 	: Super(ObjectInitializer)
@@ -26,23 +19,7 @@ UTemplateSequenceFactoryNew::UTemplateSequenceFactoryNew(const FObjectInitialize
 
 UObject* UTemplateSequenceFactoryNew::FactoryCreateNew(UClass* Class, UObject* InParent, FName Name, EObjectFlags Flags, UObject* Context, FFeedbackContext* Warn)
 {
-	auto NewTemplateSequence = NewObject<UTemplateSequence>(InParent, Name, Flags | RF_Transactional);
-	NewTemplateSequence->Initialize();
-
-	// Set up some sensible defaults
-	const UMovieSceneToolsProjectSettings* ProjectSettings = GetDefault<UMovieSceneToolsProjectSettings>();
-
-	FFrameRate TickResolution = NewTemplateSequence->GetMovieScene()->GetTickResolution();
-	NewTemplateSequence->GetMovieScene()->SetPlaybackRange((ProjectSettings->DefaultStartTime*TickResolution).FloorToFrame(), (ProjectSettings->DefaultDuration*TickResolution).FloorToFrame().Value);
-
-	if (BoundActorClass)
-	{
-		NewTemplateSequence->BoundActorClass = BoundActorClass;
-
-		InitializeSpawnable(NewTemplateSequence, BoundActorClass);
-	}
-
-	return NewTemplateSequence;
+	return FTemplateSequenceFactoryUtil::CreateTemplateSequence(InParent, Name, Flags, UTemplateSequence::StaticClass(), BoundActorClass);
 }
 
 bool UTemplateSequenceFactoryNew::ConfigureProperties()
@@ -70,56 +47,4 @@ bool UTemplateSequenceFactoryNew::ShouldShowInNewMenu() const
 	return true;
 }
 
-void UTemplateSequenceFactoryNew::InitializeSpawnable(UMovieSceneSequence* InTemplateSequence, TSubclassOf<UObject> InObjectTemplateClass)
-{
-	// Grab the object spawners, so we can create a new actor of the desired class.
-	TArray<TSharedRef<IMovieSceneObjectSpawner>> ObjectSpawners;
-	ILevelSequenceModule& LevelSequenceModule = FModuleManager::LoadModuleChecked<ILevelSequenceModule>("LevelSequence");
-	LevelSequenceModule.GenerateObjectSpawners(ObjectSpawners);
-
-	// Run the object spawners until we find one that works.
-	UMovieScene* MovieScene = InTemplateSequence->GetMovieScene();
-	check(MovieScene);
-
-	TValueOrError<FNewSpawnable, FText> Result = MakeError(LOCTEXT("NoSpawnerFound", "No spawner found to create new spawnable type"));
-
-	for (TSharedRef<IMovieSceneObjectSpawner> Spawner : ObjectSpawners)
-	{
-		TValueOrError<FNewSpawnable, FText> CurResult = Spawner->CreateNewSpawnableType(*InObjectTemplateClass.Get(), *MovieScene, nullptr);
-		if (CurResult.IsValid())
-		{
-			Result = CurResult;
-			break;
-		}
-	}
-	
-	// Bail out if we haven't found any good spawner.
-	if (!Result.IsValid())
-	{
-		FNotificationInfo Info(Result.GetError());
-		Info.ExpireDuration = 3.0f;
-		FSlateNotificationManager::Get().AddNotification(Info);
-
-		return;
-	}
-
-	// Create the spawnable.
-	FNewSpawnable& NewSpawnable = Result.GetValue();
-
-	NewSpawnable.Name = MovieSceneHelpers::MakeUniqueSpawnableName(MovieScene, NewSpawnable.Name);			
-
-	const FGuid NewGuid = MovieScene->AddSpawnable(NewSpawnable.Name, *NewSpawnable.ObjectTemplate);
-	
-	// Create a new spawn section.
-	if (ensure(NewGuid.IsValid()))
-	{
-		UMovieSceneSpawnTrack* NewSpawnTrack = Cast<UMovieSceneSpawnTrack>(MovieScene->AddTrack(UMovieSceneSpawnTrack::StaticClass(), NewGuid));
-		UMovieSceneBoolSection* NewSpawnSection = Cast<UMovieSceneBoolSection>(NewSpawnTrack->CreateNewSection());
-		NewSpawnSection->GetChannel().SetDefault(true);
-		NewSpawnSection->SetRange(TRange<FFrameNumber>::All());
-		NewSpawnTrack->AddSection(*NewSpawnSection);
-		NewSpawnTrack->SetObjectId(NewGuid);
-	}
-}
-
 #undef LOCTEXT_NAMESPACE
diff --git a/Engine/Plugins/MovieScene/TemplateSequence/Source/TemplateSequenceEditor/Private/Factories/TemplateSequenceFactoryNew.h b/Engine/Plugins/MovieScene/TemplateSequence/Source/TemplateSequenceEditor/Private/Factories/TemplateSequenceFactoryNew.h
--- a/Engine/Plugins/MovieScene/TemplateSequence/Source/TemplateSequenceEditor/Private/Factories/TemplateSequenceFactoryNew.h
+++ b/Engine/Plugins/MovieScene/TemplateSequence/Source/TemplateSequenceEditor/Private/Factories/TemplateSequenceFactoryNew.h
@@ -22,9 +22,6 @@ class UTemplateSequenceFactoryNew : public UFactory
 	virtual bool ConfigureProperties() override;
 	virtual bool ShouldShowInNewMenu() const override;
 
-	// Create a new spawnable for the chosen root object class.
-	static void InitializeSpawnable(UMovieSceneSequence* InTemplateSequence, TSubclassOf<UObject> InObjectTemplateClass);
-
 	// The root object binding class of the created template sequence.
 	UPROPERTY(EditAnywhere, Category=TemplateSequenceFactory)
 	TSubclassOf<UObject> BoundActorClass;
diff --git a/Engine/Plugins/MovieScene/TemplateSequence/Source/TemplateSequenceEditor/Private/Factories/TemplateSequenceFactoryUtil.cpp b/Engine/Plugins/MovieScene/TemplateSequence/Source/TemplateSequenceEditor/Private/Factories/TemplateSequenceFactoryUtil.cpp
new file mode 100644
--- /dev/null
+++ b/Engine/Plugins/MovieScene/TemplateSequence/Source/TemplateSequenceEditor/Private/Factories/TemplateSequenceFactoryUtil.cpp
@@ -0,0 +1,89 @@
+// Copyright Epic Games, Inc. All Rights Reserved.
+
+#include "Factories/TemplateSequenceFactoryUtil.h"
+#include "Framework/Notifications/NotificationManager.h"
+#include "ILevelSequenceModule.h"
+#include "Modules/ModuleManager.h"
+#include "MovieScene.h"
+#include "MovieSceneToolsProjectSettings.h"
+#include "Sections/MovieSceneBoolSection.h"
+#include "TemplateSequence.h"
+#include "Tracks/MovieSceneSpawnTrack.h"
+#include "Widgets/Notifications/SNotificationList.h"
+
+#define LOCTEXT_NAMESPACE "TemplateSequenceFactoryUtil"
+
+UTemplateSequence* FTemplateSequenceFactoryUtil::CreateTemplateSequence(UObject* InParent, FName InName, EObjectFlags InFlags, TSubclassOf<UTemplateSequence> InSequenceClass, TSubclassOf<UObject> InObjectTemplateClass)
+{
+	UTemplateSequence* NewTemplateSequence = NewObject<UTemplateSequence>(InParent, InSequenceClass, InName, InFlags | RF_Transactional);
+	NewTemplateSequence->Initialize();
+
+	// Set up some sensible defaults
+	const UMovieSceneToolsProjectSettings* ProjectSettings = GetDefault<UMovieSceneToolsProjectSettings>();
+
+	FFrameRate TickResolution = NewTemplateSequence->GetMovieScene()->GetTickResolution();
+	NewTemplateSequence->GetMovieScene()->SetPlaybackRange((ProjectSettings->DefaultStartTime*TickResolution).FloorToFrame(), (ProjectSettings->DefaultDuration*TickResolution).FloorToFrame().Value);
+
+	if (InObjectTemplateClass)
+	{
+		NewTemplateSequence->BoundActorClass = InObjectTemplateClass;
+
+		InitializeSpawnable(NewTemplateSequence, InObjectTemplateClass);
+	}
+
+	return NewTemplateSequence;
+}
+
+void FTemplateSequenceFactoryUtil::InitializeSpawnable(UMovieSceneSequence* InTemplateSequence, TSubclassOf<UObject> InObjectTemplateClass)
+{
+	// Grab the object spawners, so we can create a new actor of the desired class.
+	TArray<TSharedRef<IMovieSceneObjectSpawner>> ObjectSpawners;
+	ILevelSequenceModule& LevelSequenceModule = FModuleManager::LoadModuleChecked<ILevelSequenceModule>("LevelSequence");
+	LevelSequenceModule.GenerateObjectSpawners(ObjectSpawners);
+
+	// Run the object spawners until we find one that works.
+	UMovieScene* MovieScene = InTemplateSequence->GetMovieScene();
+	check(MovieScene);
+
+	TValueOrError<FNewSpawnable, FText> Result = MakeError(LOCTEXT("NoSpawnerFound", "No spawner found to create new spawnable type"));
+
+	for (TSharedRef<IMovieSceneObjectSpawner> Spawner : ObjectSpawners)
+	{
+		TValueOrError<FNewSpawnable, FText> CurResult = Spawner->CreateNewSpawnableType(*InObjectTemplateClass.Get(), *MovieScene, nullptr);
+		if (CurResult.IsValid())
+		{
+			Result = CurResult;
+			break;
+		}
+	}
+	
+	// Bail out if we haven't found any good spawner.
+	if (!Result.IsValid())
+	{
+		FNotificationInfo Info(Result.GetError());
+		Info.ExpireDuration = 3.0f;
+		FSlateNotificationManager::Get().AddNotification(Info);
+
+		return;
+	}
+
+	// Create the spawnable.
+	FNewSpawnable& NewSpawnable = Result.GetValue();
+
+	NewSpawnable.Name = MovieSceneHelpers::MakeUniqueSpawnableName(MovieScene, NewSpawnable.Name);			
+
+	const FGuid NewGuid = MovieScene->AddSpawnable(NewSpawnable.Name, *NewSpawnable.ObjectTemplate);
+	
+	// Create a new spawn section.
+	if (ensure(NewGuid.IsValid()))
+	{
+		UMovieSceneSpawnTrack* NewSpawnTrack = Cast<UMovieSceneSpawnTrack>(MovieScene->AddTrack(UMovieSceneSpawnTrack::StaticClass(), NewGuid));
+		UMovieSceneBoolSection* NewSpawnSection = Cast<UMovieSceneBoolSection>(NewSpawnTrack->CreateNewSection());
+		NewSpawnSection->GetChannel().SetDefault(true);
+		NewSpawnSection->SetRange(TRange<FFrameNumber>::All());
+		NewSpawnTrack->AddSection(*NewSpawnSection);
+		NewSpawnTrack->SetObjectId(NewGuid);
+	}
+}
+
+#undef LOCTEXT_NAMESPACE
diff --git a/Engine/Plugins/MovieScene/TemplateSequence/Source/TemplateSequenceEditor/Private/Factories/TemplateSequenceFactoryUtil.h b/Engine/Plugins/MovieScene/TemplateSequence/Source/TemplateSequenceEditor/Private/Factories/TemplateSequenceFactoryUtil.h
new file mode 100644
--- /dev/null
+++ b/Engine/Plugins/MovieScene/TemplateSequence/Source/TemplateSequenceEditor/Private/Factories/TemplateSequenceFactoryUtil.h
@@ -0,0 +1,28 @@
+// Copyright Epic Games, Inc. All Rights Reserved.
+
+#pragma once
+
+#include "CoreMinimal.h"
+#include "Templates/SubclassOf.h"
+#include "UObject/ObjectMacros.h"
+
+class UMovieSceneSequence;
+class UObject;
+class UTemplateSequence;
+
+/**
+ * Utility class for template sequence factories
+ */
+class FTemplateSequenceFactoryUtil
+{
+public:
+
+	// Create a new template sequence
+	static UTemplateSequence* CreateTemplateSequence(UObject* InParent, FName InName, EObjectFlags Flags, TSubclassOf<UTemplateSequence> InSequenceClass, TSubclassOf<UObject> InObjectTemplateClass);
+
+private:
+
+	// Create a new spawnable for the chosen root object class.
+	static void InitializeSpawnable(UMovieSceneSequence* InTemplateSequence, TSubclassOf<UObject> InObjectTemplateClass);
+};
+
diff --git a/Engine/Plugins/MovieScene/TemplateSequence/Source/TemplateSequenceEditor/Private/TemplateSequenceEditorToolkit.cpp b/Engine/Plugins/MovieScene/TemplateSequence/Source/TemplateSequenceEditor/Private/TemplateSequenceEditorToolkit.cpp
--- a/Engine/Plugins/MovieScene/TemplateSequence/Source/TemplateSequenceEditor/Private/TemplateSequenceEditorToolkit.cpp
+++ b/Engine/Plugins/MovieScene/TemplateSequence/Source/TemplateSequenceEditor/Private/TemplateSequenceEditorToolkit.cpp
@@ -15,7 +15,6 @@
 #include "LevelEditorSequencerIntegration.h"
 #include "Misc/TemplateSequenceEditorPlaybackContext.h"
 #include "Misc/TemplateSequenceEditorSpawnRegister.h"
-#include "Misc/TemplateSequenceEditorUtil.h"
 #include "Modules/ModuleManager.h"
 #include "ScopedTransaction.h"
 #include "SequencerSettings.h"
@@ -111,12 +110,6 @@ void FTemplateSequenceEditorToolkit::Initialize(const EToolkitMode::Type Mode, c
 	SpawnRegister->SetSequencer(Sequencer);
 	Sequencer->OnActorAddedToSequencer().AddSP(this, &FTemplateSequenceEditorToolkit::HandleActorAddedToSequencer);
 
-	if (ToolkitParams.InitialBindingClass != nullptr)
-	{
-		FTemplateSequenceEditorUtil Util(TemplateSequence, *Sequencer.Get());
-		Util.ChangeActorBinding(ToolkitParams.InitialBindingClass);
-	}
-
 	FLevelEditorSequencerIntegrationOptions Options;
 	Options.bRequiresLevelEvents = true;
 	Options.bRequiresActorEvents = true;
diff --git a/Engine/Plugins/MovieScene/TemplateSequence/Source/TemplateSequenceEditor/Private/TemplateSequenceEditorToolkit.h b/Engine/Plugins/MovieScene/TemplateSequence/Source/TemplateSequenceEditor/Private/TemplateSequenceEditorToolkit.h
--- a/Engine/Plugins/MovieScene/TemplateSequence/Source/TemplateSequenceEditor/Private/TemplateSequenceEditorToolkit.h
+++ b/Engine/Plugins/MovieScene/TemplateSequence/Source/TemplateSequenceEditor/Private/TemplateSequenceEditorToolkit.h
@@ -18,7 +18,6 @@ class FActorDragDropGraphEdOp;
 struct FTemplateSequenceToolkitParams
 {
 	bool bCanChangeBinding = true;
-	UClass* InitialBindingClass = nullptr;
 };
 
 /**
