Changed template sequence actors to use asynchronous updates

diff --git a/Engine/Plugins/MovieScene/TemplateSequence/Source/TemplateSequence/Private/TemplateSequenceActor.cpp b/Engine/Plugins/MovieScene/TemplateSequence/Source/TemplateSequence/Private/TemplateSequenceActor.cpp
--- a/Engine/Plugins/MovieScene/TemplateSequence/Source/TemplateSequence/Private/TemplateSequenceActor.cpp
+++ b/Engine/Plugins/MovieScene/TemplateSequence/Source/TemplateSequence/Private/TemplateSequenceActor.cpp
@@ -77,7 +77,7 @@
 {
 	if (SequencePlayer)
 	{
-		SequencePlayer->Update(DeltaSeconds);
+		SequencePlayer->UpdateAsync(DeltaSeconds);
 	}
 }
 
