Sequencer: Lazy-create the linker for the camera animation sub-system.

This simplifies debugging since we have less linkers being created in simple test levels.

diff --git a/Engine/Plugins/MovieScene/TemplateSequence/Source/TemplateSequence/Private/CameraAnimationSequenceSubsystem.cpp b/Engine/Plugins/MovieScene/TemplateSequence/Source/TemplateSequence/Private/CameraAnimationSequenceSubsystem.cpp
--- a/Engine/Plugins/MovieScene/TemplateSequence/Source/TemplateSequence/Private/CameraAnimationSequenceSubsystem.cpp
+++ b/Engine/Plugins/MovieScene/TemplateSequence/Source/TemplateSequence/Private/CameraAnimationSequenceSubsystem.cpp
@@ -1,48 +1,50 @@
-// Copyright Epic Games, Inc. All Rights Reserved.
-
-#include "CameraAnimationSequenceSubsystem.h"
-#include "Engine/World.h"
-#include "EntitySystem/MovieSceneEntitySystemLinker.h"
-
-#define LOCTEXT_NAMESPACE "CameraAnimationSequenceSubsystem"
-
-UCameraAnimationSequenceSubsystem* UCameraAnimationSequenceSubsystem::GetCameraAnimationSequenceSubsystem(const UWorld* InWorld)
-{
-	if (InWorld)
-	{
-		return InWorld->GetSubsystem<UCameraAnimationSequenceSubsystem>();
-	}
-
-	return nullptr;
-}
-
-UCameraAnimationSequenceSubsystem::UCameraAnimationSequenceSubsystem()
-{
-}
-
-UCameraAnimationSequenceSubsystem::~UCameraAnimationSequenceSubsystem()
-{
-}
-
-void UCameraAnimationSequenceSubsystem::Initialize(FSubsystemCollectionBase& Collection)
-{
-	Super::Initialize(Collection);
-
-	Linker = NewObject<UMovieSceneEntitySystemLinker>(this, TEXT("CameraAnimationSequenceSubsystemLinker"));
-	Runner.AttachToLinker(Linker);
-}
-
-void UCameraAnimationSequenceSubsystem::Deinitialize()
-{
-	// We check if the runner still has a valid pointer on the linker because the linker could
-	// have been GC'ed just now, which would make DetachFromLinker complain.
-	if (Runner.GetLinker())
-	{
-		Runner.DetachFromLinker();
-	}
-	Linker = nullptr;
-
-	Super::Deinitialize();
-}
-
-#undef LOCTEXT_NAMESPACE
+// Copyright Epic Games, Inc. All Rights Reserved.
+
+#include "CameraAnimationSequenceSubsystem.h"
+#include "Engine/World.h"
+#include "EntitySystem/MovieSceneEntitySystemLinker.h"
+
+#define LOCTEXT_NAMESPACE "CameraAnimationSequenceSubsystem"
+
+UCameraAnimationSequenceSubsystem* UCameraAnimationSequenceSubsystem::GetCameraAnimationSequenceSubsystem(const UWorld* InWorld)
+{
+	if (InWorld)
+	{
+		return InWorld->GetSubsystem<UCameraAnimationSequenceSubsystem>();
+	}
+
+	return nullptr;
+}
+
+UCameraAnimationSequenceSubsystem::UCameraAnimationSequenceSubsystem()
+{
+}
+
+UCameraAnimationSequenceSubsystem::~UCameraAnimationSequenceSubsystem()
+{
+}
+
+void UCameraAnimationSequenceSubsystem::Deinitialize()
+{
+	// We check if the runner still has a valid pointer on the linker because the linker could
+	// have been GC'ed just now, which would make DetachFromLinker complain.
+	if (Runner.GetLinker())
+	{
+		Runner.DetachFromLinker();
+	}
+	Linker = nullptr;
+
+	Super::Deinitialize();
+}
+
+UMovieSceneEntitySystemLinker* UCameraAnimationSequenceSubsystem::GetLinker(bool bAutoCreate)
+{
+	if (!Linker && bAutoCreate)
+	{
+		Linker = NewObject<UMovieSceneEntitySystemLinker>(this, TEXT("CameraAnimationSequenceSubsystemLinker"));
+		Runner.AttachToLinker(Linker);
+	}
+	return Linker;
+}
+
+#undef LOCTEXT_NAMESPACE
diff --git a/Engine/Plugins/MovieScene/TemplateSequence/Source/TemplateSequence/Public/CameraAnimationSequenceSubsystem.h b/Engine/Plugins/MovieScene/TemplateSequence/Source/TemplateSequence/Public/CameraAnimationSequenceSubsystem.h
--- a/Engine/Plugins/MovieScene/TemplateSequence/Source/TemplateSequence/Public/CameraAnimationSequenceSubsystem.h
+++ b/Engine/Plugins/MovieScene/TemplateSequence/Source/TemplateSequence/Public/CameraAnimationSequenceSubsystem.h
@@ -26,11 +26,10 @@
 	virtual ~UCameraAnimationSequenceSubsystem();
 
 	/** Gets the Sequencer linker owned by this sybsystem */
-	UMovieSceneEntitySystemLinker* GetLinker() { return Linker; }
+	UMovieSceneEntitySystemLinker* GetLinker(bool bAutoCreate = true);
 
 private:
 	// Begin USubsystem
-	virtual void Initialize(FSubsystemCollectionBase& Collection) override;
 	virtual void Deinitialize() override;
 	// End USubsystem
 	
