Fixing non-unity build

diff --git a/Engine/Plugins/Editor/AssetSearch/Source/Public/SearchQuery.h b/Engine/Plugins/Editor/AssetSearch/Source/Public/SearchQuery.h
--- a/Engine/Plugins/Editor/AssetSearch/Source/Public/SearchQuery.h
+++ b/Engine/Plugins/Editor/AssetSearch/Source/Public/SearchQuery.h
@@ -5,6 +5,7 @@
 #include "CoreMinimal.h"
 #include "HAL/CriticalSection.h"
 #include "Misc/ScopeLock.h"
+#include "UObject/TopLevelAssetPath.h"
 
 struct FSearchRecord
 {
diff --git a/Engine/Plugins/MovieScene/TemplateSequence/Source/TemplateSequence/Private/TemplateSequence.cpp b/Engine/Plugins/MovieScene/TemplateSequence/Source/TemplateSequence/Private/TemplateSequence.cpp
--- a/Engine/Plugins/MovieScene/TemplateSequence/Source/TemplateSequence/Private/TemplateSequence.cpp
+++ b/Engine/Plugins/MovieScene/TemplateSequence/Source/TemplateSequence/Private/TemplateSequence.cpp
@@ -1,6 +1,7 @@
 // Copyright Epic Games, Inc. All Rights Reserved.
 
 #include "TemplateSequence.h"
+#include "AssetRegistry/AssetData.h"
 #include "Engine/Engine.h"
 #include "Engine/World.h"
 #include "Components/ActorComponent.h"
diff --git a/Engine/Source/Runtime/CoreUObject/Private/UObject/ObjectRedirector.cpp b/Engine/Source/Runtime/CoreUObject/Private/UObject/ObjectRedirector.cpp
--- a/Engine/Source/Runtime/CoreUObject/Private/UObject/ObjectRedirector.cpp
+++ b/Engine/Source/Runtime/CoreUObject/Private/UObject/ObjectRedirector.cpp
@@ -10,6 +10,7 @@
 #include "UObject/Package.h"
 #include "Templates/Casts.h"
 #include "UObject/PropertyPortFlags.h"
+#include "UObject/UnrealType.h"
 
 /*-----------------------------------------------------------------------------
 	UObjectRedirector
