Try to find the full path name of template sequences' bound actor classes if only the short name was saved

diff --git a/Engine/Plugins/MovieScene/TemplateSequence/Source/TemplateSequenceEditor/Private/TrackEditors/TemplateSequenceTrackEditor.cpp b/Engine/Plugins/MovieScene/TemplateSequence/Source/TemplateSequenceEditor/Private/TrackEditors/TemplateSequenceTrackEditor.cpp
--- a/Engine/Plugins/MovieScene/TemplateSequence/Source/TemplateSequenceEditor/Private/TrackEditors/TemplateSequenceTrackEditor.cpp
+++ b/Engine/Plugins/MovieScene/TemplateSequence/Source/TemplateSequenceEditor/Private/TrackEditors/TemplateSequenceTrackEditor.cpp
@@ -172,13 +172,29 @@ public:
 				{
 					if (LegacyBaseClass == nullptr || AssetData.AssetClassPath != LegacyBaseClass->GetClassPathName())
 					{
-						const FAssetDataTagMapSharedView::FFindTagResult FoundBoundActorClass = AssetData.TagsAndValues.FindTag("BoundActorClass");
-						if (FoundBoundActorClass.IsSet())
+						FString FoundBoundActorClass;
+						const FAssetDataTagMapSharedView::FFindTagResult FoundBoundActorClassTag = AssetData.TagsAndValues.FindTag("BoundActorClass");
+						if (FoundBoundActorClassTag.IsSet())
+						{
+							FString TagValue = FoundBoundActorClassTag.GetValue();
+							if (FPackageName::IsShortPackageName(TagValue))
+							{
+								// Replace the short name with the full path name if possible.
+								FTopLevelAssetPath BoundActorClassPath = UClass::TryConvertShortTypeNameToPathName<UStruct>(TagValue, ELogVerbosity::Warning, TEXT("STemplateSequenceAssetSubMenu"));
+								if (!BoundActorClassPath.IsNull())
+								{
+									FoundBoundActorClass = BoundActorClassPath.ToString();
+								}
+							}
+							else
+							{
+								FoundBoundActorClass = TagValue;
+							}
+						}
+						if (!FoundBoundActorClass.IsEmpty())
 						{
 							// Filter this out if it's got an incompatible bound actor class.
-							const FString FoundBoundActorClassName(FoundBoundActorClass.GetValue());
-							ensureAlwaysMsgf(!FPackageName::IsShortPackageName(FoundBoundActorClassName), TEXT("BoundActorClass tag contains short name \"%s\". Path name expected. See UTemplateSequence::GetAssetRegistryTags."), *FoundBoundActorClassName);
-							return !BaseClassNames.Contains(FoundBoundActorClassName);
+							return !BaseClassNames.Contains(FoundBoundActorClass);
 						}
 						else
 						{
