Move Level Sequence, Template Sequence, and Camera Animation Sequence asset types into a new 'Cinematics' top-level asset category.

diff --git a/Engine/Plugins/MovieScene/LevelSequenceEditor/Source/LevelSequenceEditor/Private/AssetTools/LevelSequenceActions.cpp b/Engine/Plugins/MovieScene/LevelSequenceEditor/Source/LevelSequenceEditor/Private/AssetTools/LevelSequenceActions.cpp
--- a/Engine/Plugins/MovieScene/LevelSequenceEditor/Source/LevelSequenceEditor/Private/AssetTools/LevelSequenceActions.cpp
+++ b/Engine/Plugins/MovieScene/LevelSequenceEditor/Source/LevelSequenceEditor/Private/AssetTools/LevelSequenceActions.cpp
@@ -12,8 +12,9 @@
 /* FLevelSequenceActions constructors
  *****************************************************************************/
 
-FLevelSequenceActions::FLevelSequenceActions( const TSharedRef<ISlateStyle>& InStyle )
+FLevelSequenceActions::FLevelSequenceActions( const TSharedRef<ISlateStyle>& InStyle, const EAssetTypeCategories::Type InCategory)
 	: Style(InStyle)
+	, AssetCategory(InCategory)
 { }
 
 
@@ -22,17 +23,7 @@ FLevelSequenceActions::FLevelSequenceActions( const TSharedRef<ISlateStyle>& InS
 
 uint32 FLevelSequenceActions::GetCategories()
 {
-	return EAssetTypeCategories::Animation;
-}
-
-
-const TArray<FText>& FLevelSequenceActions::GetSubMenus() const
-{
-	static const TArray<FText> SubMenus
-	{
-		LOCTEXT("AnimSequencerSubMenu", "Sequencer")
-	};
-	return SubMenus;
+	return AssetCategory;
 }
 
 
diff --git a/Engine/Plugins/MovieScene/LevelSequenceEditor/Source/LevelSequenceEditor/Private/AssetTools/LevelSequenceActions.h b/Engine/Plugins/MovieScene/LevelSequenceEditor/Source/LevelSequenceEditor/Private/AssetTools/LevelSequenceActions.h
--- a/Engine/Plugins/MovieScene/LevelSequenceEditor/Source/LevelSequenceEditor/Private/AssetTools/LevelSequenceActions.h
+++ b/Engine/Plugins/MovieScene/LevelSequenceEditor/Source/LevelSequenceEditor/Private/AssetTools/LevelSequenceActions.h
@@ -6,6 +6,7 @@
 #include "Styling/ISlateStyle.h"
 #include "Toolkits/IToolkitHost.h"
 #include "AssetTypeActions_Base.h"
+#include "AssetTypeCategories.h"
 
 
 /**
@@ -21,14 +22,13 @@ public:
 	 *
 	 * @param InStyle The style set to use for asset editor toolkits.
 	 */
-	FLevelSequenceActions(const TSharedRef<ISlateStyle>& InStyle);
+	FLevelSequenceActions(const TSharedRef<ISlateStyle>& InStyle, const EAssetTypeCategories::Type InCategory);
 
 public:
 	
 	// IAssetTypeActions interface
 
 	virtual uint32 GetCategories() override;
-	virtual const TArray<FText>& GetSubMenus() const override;
 	virtual FText GetName() const override;
 	virtual UClass* GetSupportedClass() const override;
 	virtual FColor GetTypeColor() const override;
@@ -41,4 +41,6 @@ private:
 
 	/** Pointer to the style set to use for toolkits. */
 	TSharedRef<ISlateStyle> Style;
+
+	const EAssetTypeCategories::Type AssetCategory;
 };
diff --git a/Engine/Plugins/MovieScene/LevelSequenceEditor/Source/LevelSequenceEditor/Private/LevelSequenceEditorModule.cpp b/Engine/Plugins/MovieScene/LevelSequenceEditor/Source/LevelSequenceEditor/Private/LevelSequenceEditorModule.cpp
--- a/Engine/Plugins/MovieScene/LevelSequenceEditor/Source/LevelSequenceEditor/Private/LevelSequenceEditorModule.cpp
+++ b/Engine/Plugins/MovieScene/LevelSequenceEditor/Source/LevelSequenceEditor/Private/LevelSequenceEditorModule.cpp
@@ -122,7 +122,8 @@ protected:
 	void RegisterAssetTools()
 	{
 		IAssetTools& AssetTools = FModuleManager::LoadModuleChecked<FAssetToolsModule>("AssetTools").Get();
-		LevelSequenceTypeActions = MakeShared<FLevelSequenceActions>(FLevelSequenceEditorStyle::Get());
+		EAssetTypeCategories::Type CinematicAssetCategoryBit = AssetTools.RegisterAdvancedAssetCategory(FName(TEXT("Cinematics")), LOCTEXT("CinematicsAssetCategory", "Cinematics"));
+		LevelSequenceTypeActions = MakeShared<FLevelSequenceActions>(FLevelSequenceEditorStyle::Get(), CinematicAssetCategoryBit);
 		AssetTools.RegisterAssetTypeActions(LevelSequenceTypeActions.ToSharedRef());
 	}
 
diff --git a/Engine/Plugins/MovieScene/TemplateSequence/Source/TemplateSequenceEditor/Private/AssetTools/CameraAnimationSequenceActions.cpp b/Engine/Plugins/MovieScene/TemplateSequence/Source/TemplateSequenceEditor/Private/AssetTools/CameraAnimationSequenceActions.cpp
--- a/Engine/Plugins/MovieScene/TemplateSequence/Source/TemplateSequenceEditor/Private/AssetTools/CameraAnimationSequenceActions.cpp
+++ b/Engine/Plugins/MovieScene/TemplateSequence/Source/TemplateSequenceEditor/Private/AssetTools/CameraAnimationSequenceActions.cpp
@@ -7,8 +7,8 @@
 
 #define LOCTEXT_NAMESPACE "AssetTypeActions"
 
-FCameraAnimationSequenceActions::FCameraAnimationSequenceActions(const TSharedRef<ISlateStyle>& InStyle)
-    : FTemplateSequenceActions(InStyle)
+FCameraAnimationSequenceActions::FCameraAnimationSequenceActions(const TSharedRef<ISlateStyle>& InStyle, const EAssetTypeCategories::Type InAssetCategory)
+    : FTemplateSequenceActions(InStyle, InAssetCategory)
 {
 }
 
diff --git a/Engine/Plugins/MovieScene/TemplateSequence/Source/TemplateSequenceEditor/Private/AssetTools/CameraAnimationSequenceActions.h b/Engine/Plugins/MovieScene/TemplateSequence/Source/TemplateSequenceEditor/Private/AssetTools/CameraAnimationSequenceActions.h
--- a/Engine/Plugins/MovieScene/TemplateSequence/Source/TemplateSequenceEditor/Private/AssetTools/CameraAnimationSequenceActions.h
+++ b/Engine/Plugins/MovieScene/TemplateSequence/Source/TemplateSequenceEditor/Private/AssetTools/CameraAnimationSequenceActions.h
@@ -3,12 +3,13 @@
 #pragma once
 
 #include "TemplateSequenceActions.h"
+#include "AssetTypeCategories.h"
 
 class FCameraAnimationSequenceActions : public FTemplateSequenceActions
 {
 public:
     
-    FCameraAnimationSequenceActions(const TSharedRef<ISlateStyle>& InStyle);
+    FCameraAnimationSequenceActions(const TSharedRef<ISlateStyle>& InStyle, const EAssetTypeCategories::Type InAssetCategory);
 
     virtual FText GetName() const override;
     virtual UClass* GetSupportedClass() const override;
diff --git a/Engine/Plugins/MovieScene/TemplateSequence/Source/TemplateSequenceEditor/Private/AssetTools/TemplateSequenceActions.cpp b/Engine/Plugins/MovieScene/TemplateSequence/Source/TemplateSequenceEditor/Private/AssetTools/TemplateSequenceActions.cpp
--- a/Engine/Plugins/MovieScene/TemplateSequence/Source/TemplateSequenceEditor/Private/AssetTools/TemplateSequenceActions.cpp
+++ b/Engine/Plugins/MovieScene/TemplateSequence/Source/TemplateSequenceEditor/Private/AssetTools/TemplateSequenceActions.cpp
@@ -8,22 +8,14 @@
 
 #define LOCTEXT_NAMESPACE "AssetTypeActions"
 
-FTemplateSequenceActions::FTemplateSequenceActions(const TSharedRef<ISlateStyle>& InStyle)
+FTemplateSequenceActions::FTemplateSequenceActions(const TSharedRef<ISlateStyle>& InStyle, const EAssetTypeCategories::Type InAssetCategory)
 	: Style(InStyle)
+	, AssetCategory(InAssetCategory)
 { }
 
 uint32 FTemplateSequenceActions::GetCategories()
 {
-	return EAssetTypeCategories::Animation;
-}
-
-const TArray<FText>& FTemplateSequenceActions::GetSubMenus() const
-{
-	static const TArray<FText> SubMenus
-	{
-		LOCTEXT("AnimSequencerSubMenu", "Sequencer")
-	};
-	return SubMenus;
+	return AssetCategory;
 }
 
 FText FTemplateSequenceActions::GetName() const
diff --git a/Engine/Plugins/MovieScene/TemplateSequence/Source/TemplateSequenceEditor/Private/AssetTools/TemplateSequenceActions.h b/Engine/Plugins/MovieScene/TemplateSequence/Source/TemplateSequenceEditor/Private/AssetTools/TemplateSequenceActions.h
--- a/Engine/Plugins/MovieScene/TemplateSequence/Source/TemplateSequenceEditor/Private/AssetTools/TemplateSequenceActions.h
+++ b/Engine/Plugins/MovieScene/TemplateSequence/Source/TemplateSequenceEditor/Private/AssetTools/TemplateSequenceActions.h
@@ -6,6 +6,7 @@
 #include "Styling/ISlateStyle.h"
 #include "Toolkits/IToolkitHost.h"
 #include "AssetTypeActions_Base.h"
+#include "AssetTypeCategories.h"
 
 struct FTemplateSequenceToolkitParams;
 
@@ -21,13 +22,12 @@ public:
 	 *
 	 * @param InStyle The style set to use for asset editor toolkits.
 	 */
-	FTemplateSequenceActions(const TSharedRef<ISlateStyle>& InStyle);
+	FTemplateSequenceActions(const TSharedRef<ISlateStyle>& InStyle, const EAssetTypeCategories::Type InAssetCategory);
 
 public:
 
 	// IAssetTypeActions interface
 	virtual uint32 GetCategories() override;
-	virtual const TArray<FText>& GetSubMenus() const override;
 	virtual FText GetName() const override;
 	virtual UClass* GetSupportedClass() const override;
 	virtual FColor GetTypeColor() const override;
@@ -44,4 +44,6 @@ private:
 
 	/** Pointer to the style set to use for toolkits. */
 	TSharedRef<ISlateStyle> Style;
+
+	const EAssetTypeCategories::Type AssetCategory;
 };
diff --git a/Engine/Plugins/MovieScene/TemplateSequence/Source/TemplateSequenceEditor/Private/TemplateSequenceEditorModule.cpp b/Engine/Plugins/MovieScene/TemplateSequence/Source/TemplateSequenceEditor/Private/TemplateSequenceEditorModule.cpp
--- a/Engine/Plugins/MovieScene/TemplateSequence/Source/TemplateSequenceEditor/Private/TemplateSequenceEditorModule.cpp
+++ b/Engine/Plugins/MovieScene/TemplateSequence/Source/TemplateSequenceEditor/Private/TemplateSequenceEditorModule.cpp
@@ -98,11 +98,12 @@ private:
 	void RegisterAssetTools()
 	{
 		IAssetTools& AssetTools = FModuleManager::LoadModuleChecked<FAssetToolsModule>("AssetTools").Get();
+		EAssetTypeCategories::Type CinematicAssetCategoryBit = AssetTools.RegisterAdvancedAssetCategory(FName(TEXT("Cinematics")), LOCTEXT("CinematicsAssetCategory", "Cinematics"));
 
-		TemplateSequenceTypeActions = MakeShared<FTemplateSequenceActions>(FTemplateSequenceEditorStyle::Get());
+		TemplateSequenceTypeActions = MakeShared<FTemplateSequenceActions>(FTemplateSequenceEditorStyle::Get(), CinematicAssetCategoryBit);
 		AssetTools.RegisterAssetTypeActions(TemplateSequenceTypeActions.ToSharedRef());
 
-		CameraAnimationSequenceActions = MakeShared<FCameraAnimationSequenceActions>(FTemplateSequenceEditorStyle::Get());
+		CameraAnimationSequenceActions = MakeShared<FCameraAnimationSequenceActions>(FTemplateSequenceEditorStyle::Get(), CinematicAssetCategoryBit);
 		AssetTools.RegisterAssetTypeActions(CameraAnimationSequenceActions.ToSharedRef());
 	}
 
