Template Sequence: Fix crash on teardown when property scaled instances no longer contains the requested subsequence handle.

diff --git a/Engine/Plugins/MovieScene/TemplateSequence/Source/TemplateSequence/Private/Systems/TemplateSequenceSystem.cpp b/Engine/Plugins/MovieScene/TemplateSequence/Source/TemplateSequence/Private/Systems/TemplateSequenceSystem.cpp
--- a/Engine/Plugins/MovieScene/TemplateSequence/Source/TemplateSequence/Private/Systems/TemplateSequenceSystem.cpp
+++ b/Engine/Plugins/MovieScene/TemplateSequence/Source/TemplateSequence/Private/Systems/TemplateSequenceSystem.cpp
@@ -172,7 +172,7 @@ void UTemplateSequencePropertyScalingInstantiatorSystem::OnRun(FSystemTaskPrereq
 			const FTemplateSequencePropertyScaleComponentData& PropertyScale)
 	{
 		const FInstanceHandle SubSequenceHandle = InstanceRegistry->FindRelatedInstanceHandle(InstanceHandle, PropertyScale.SubSequenceID);
-		if (ensure(SubSequenceHandle.IsValid()))
+		if (ensure(SubSequenceHandle.IsValid()) && PropertyScaledInstances.Contains(SubSequenceHandle))
 		{
 			FPropertyScaleEntityIDs& PropertyScaleEntities = PropertyScaledInstances.FindChecked(SubSequenceHandle);
 			if (PropertyScaleEntities.Remove(EntityID) > 0)
