Register/unregister camera stand-in property accessors on module load/unload

diff --git a/Engine/Plugins/Cameras/GameplayCameras/Source/GameplayCameras/Private/CameraAnimationCameraModifier.cpp b/Engine/Plugins/Cameras/GameplayCameras/Source/GameplayCameras/Private/CameraAnimationCameraModifier.cpp
--- a/Engine/Plugins/Cameras/GameplayCameras/Source/GameplayCameras/Private/CameraAnimationCameraModifier.cpp
+++ b/Engine/Plugins/Cameras/GameplayCameras/Source/GameplayCameras/Private/CameraAnimationCameraModifier.cpp
@@ -78,8 +78,6 @@ float UCameraAnimationCameraModifier::EvaluateEasing(ECameraAnimationEasingType
 UCameraAnimationCameraModifier::UCameraAnimationCameraModifier(const FObjectInitializer& ObjectInitializer)
 	: Super(ObjectInitializer)
 {
-	// Make sure we have our custom accessors registered for our stand-in class.
-	UCameraAnimationSequenceCameraStandIn::RegisterCameraStandIn();
 }
 
 FCameraAnimationHandle UCameraAnimationCameraModifier::PlayCameraAnimation(UCameraAnimationSequence* Sequence, FCameraAnimationParams Params)
diff --git a/Engine/Plugins/Cameras/GameplayCameras/Source/GameplayCameras/Private/GameplayCamerasModule.cpp b/Engine/Plugins/Cameras/GameplayCameras/Source/GameplayCameras/Private/GameplayCamerasModule.cpp
--- a/Engine/Plugins/Cameras/GameplayCameras/Source/GameplayCameras/Private/GameplayCamerasModule.cpp
+++ b/Engine/Plugins/Cameras/GameplayCameras/Source/GameplayCameras/Private/GameplayCamerasModule.cpp
@@ -3,6 +3,7 @@
 #include "GameplayCamerasModule.h"
 #include "Camera/CameraModularFeature.h"
 #include "CameraAnimationCameraModifier.h"
+#include "CameraAnimationSequencePlayer.h"
 #include "Features/IModularFeatures.h"
 #include "Modules/ModuleManager.h"
 
@@ -23,6 +24,8 @@ public:
 		{
 			IModularFeatures::Get().RegisterModularFeature(ICameraModularFeature::GetModularFeatureName(), CameraModularFeature.Get());
 		}
+
+		UCameraAnimationSequenceCameraStandIn::RegisterCameraStandIn();
 	}
 
 	virtual void ShutdownModule() override
@@ -32,6 +35,8 @@ public:
 			IModularFeatures::Get().UnregisterModularFeature(ICameraModularFeature::GetModularFeatureName(), CameraModularFeature.Get());
 			CameraModularFeature = nullptr;
 		}
+
+		UCameraAnimationSequenceCameraStandIn::UnregisterCameraStandIn();
 	}
 
 private:
diff --git a/Engine/Plugins/MovieScene/TemplateSequence/Source/TemplateSequence/Private/CameraAnimationSequencePlayer.cpp b/Engine/Plugins/MovieScene/TemplateSequence/Source/TemplateSequence/Private/CameraAnimationSequencePlayer.cpp
--- a/Engine/Plugins/MovieScene/TemplateSequence/Source/TemplateSequence/Private/CameraAnimationSequencePlayer.cpp
+++ b/Engine/Plugins/MovieScene/TemplateSequence/Source/TemplateSequence/Private/CameraAnimationSequencePlayer.cpp
@@ -68,6 +68,8 @@ void UpdateInitialPropertyValues(UMovieSceneEntitySystemLinker* Linker, const TP
 }
 }
 
+bool UCameraAnimationSequenceCameraStandIn::bRegistered(false);
+
 UCameraAnimationSequenceCameraStandIn::UCameraAnimationSequenceCameraStandIn(const FObjectInitializer& ObjInit) 
 	: Super(ObjInit) 
 {
@@ -210,13 +212,25 @@ void UCameraAnimationSequenceCameraStandIn::RegisterCameraStandIn()
 {
 	using namespace UE::MovieScene;
 
-	static bool bHasRegistered = false;
-	if (!bHasRegistered)
+	if (ensure(!bRegistered))
 	{
 		FMovieSceneTracksComponentTypes* TracksComponentTypes = FMovieSceneTracksComponentTypes::Get();
 		TracksComponentTypes->Accessors.ComponentTransform.Add(UCameraAnimationSequenceCameraStandIn::StaticClass(), TEXT("Transform"), &GetCameraStandInTransform, &SetCameraStandInTransform);
 
-		bHasRegistered = true;
+		bRegistered = true;
+	}
+}
+
+void UCameraAnimationSequenceCameraStandIn::UnregisterCameraStandIn()
+{
+	using namespace UE::MovieScene;
+
+	if (ensure(bRegistered))
+	{
+		FMovieSceneTracksComponentTypes* TracksComponentTypes = FMovieSceneTracksComponentTypes::Get();
+		TracksComponentTypes->Accessors.ComponentTransform.RemoveAll(UCameraAnimationSequenceCameraStandIn::StaticClass());
+
+		bRegistered = false;
 	}
 }
 
diff --git a/Engine/Plugins/MovieScene/TemplateSequence/Source/TemplateSequence/Private/SequenceCameraShake.cpp b/Engine/Plugins/MovieScene/TemplateSequence/Source/TemplateSequence/Private/SequenceCameraShake.cpp
--- a/Engine/Plugins/MovieScene/TemplateSequence/Source/TemplateSequence/Private/SequenceCameraShake.cpp
+++ b/Engine/Plugins/MovieScene/TemplateSequence/Source/TemplateSequence/Private/SequenceCameraShake.cpp
@@ -34,9 +34,6 @@ USequenceCameraShakePattern::USequenceCameraShakePattern(const FObjectInitialize
 {
 	CameraStandIn = CreateDefaultSubobject<UCameraAnimationSequenceCameraStandIn>(TEXT("CameraStandIn"), true);
 	Player = CreateDefaultSubobject<UCameraAnimationSequencePlayer>(TEXT("CameraShakePlayer"), true);
-
-	// Make sure we have our custom accessors registered for our stand-in class.
-	UCameraAnimationSequenceCameraStandIn::RegisterCameraStandIn();
 }
 
 void USequenceCameraShakePattern::GetShakePatternInfoImpl(FCameraShakeInfo& OutInfo) const
diff --git a/Engine/Plugins/MovieScene/TemplateSequence/Source/TemplateSequence/Public/CameraAnimationSequencePlayer.h b/Engine/Plugins/MovieScene/TemplateSequence/Source/TemplateSequence/Public/CameraAnimationSequencePlayer.h
--- a/Engine/Plugins/MovieScene/TemplateSequence/Source/TemplateSequence/Public/CameraAnimationSequencePlayer.h
+++ b/Engine/Plugins/MovieScene/TemplateSequence/Source/TemplateSequence/Public/CameraAnimationSequencePlayer.h
@@ -78,11 +78,16 @@ public:
 	/** Register the stand-in class with the sequencer ECS component registry */
 	static void RegisterCameraStandIn();
 
+	/** Unregister the stand-in class with the sequencer ECS component registry */
+	static void UnregisterCameraStandIn();
+
 private:
 	void ResetDefaultValues(const FMinimalViewInfo& ViewInfo);
 	void UpdateInitialPropertyValues(UMovieSceneEntitySystemLinker* Linker);
 
 private:
+	static bool bRegistered;
+
 	FTransform Transform;
 
 	bool bIsCineCamera = false;
diff --git a/Engine/Source/Runtime/MovieScene/Public/EntitySystem/MovieScenePropertySystemTypes.h b/Engine/Source/Runtime/MovieScene/Public/EntitySystem/MovieScenePropertySystemTypes.h
--- a/Engine/Source/Runtime/MovieScene/Public/EntitySystem/MovieScenePropertySystemTypes.h
+++ b/Engine/Source/Runtime/MovieScene/Public/EntitySystem/MovieScenePropertySystemTypes.h
@@ -234,6 +234,22 @@ struct TCustomPropertyRegistration : ICustomPropertyRegistration
 		CustomAccessors.Add(TCustomPropertyAccessor<PropertyTraits>{ ClassType, PropertyName, { Getter, Setter } });
 	}
 
+	void Remove(UClass* ClassType, FName PropertyName)
+	{
+		CustomAccessors.RemoveAll([=](const TCustomPropertyAccessor<PropertyTraits>& Accessor)
+				{
+					return Accessor.Class == ClassType && Accessor.PropertyPath == PropertyName;
+				});
+	}
+
+	void RemoveAll(UClass* ClassType)
+	{
+		CustomAccessors.RemoveAll([=](const TCustomPropertyAccessor<PropertyTraits>& Accessor)
+				{
+					return Accessor.Class == ClassType;
+				});
+	}
+
 private:
 
 	/** */
@@ -642,4 +658,4 @@ using TSetCompositePropertyValues = TSetCompositePropertyValuesImpl<PropertyTrai
 
 
 } // namespace MovieScene
-} // namespace UE
\ No newline at end of file
+} // namespace UE
