Remove Matinee (1/2)

This change completely removes Matinee from the codebase. The data types and runtime evaluation code are gone. The "InterpEdit" editor mode, asset factories, and other editor support for Matinee is gone. CameraAnims and their sequencer tracks are gone. FBX import/export for Matinee assets is gone.

diff --git a/Engine/Plugins/Animation/ControlRig/Source/ControlRigEditor/Private/Sequencer/ControlRigParameterTrackEditor.cpp b/Engine/Plugins/Animation/ControlRig/Source/ControlRigEditor/Private/Sequencer/ControlRigParameterTrackEditor.cpp
--- a/Engine/Plugins/Animation/ControlRig/Source/ControlRigEditor/Private/Sequencer/ControlRigParameterTrackEditor.cpp
+++ b/Engine/Plugins/Animation/ControlRig/Source/ControlRigEditor/Private/Sequencer/ControlRigParameterTrackEditor.cpp
@@ -25,8 +25,6 @@
 #include "AssetRegistry/AssetRegistryModule.h"
 #include "IContentBrowserSingleton.h"
 #include "ContentBrowserModule.h"
-#include "MatineeImportTools.h"
-#include "Matinee/InterpTrackAnimControl.h"
 #include "SequencerUtilities.h"
 #include "ISectionLayoutBuilder.h"
 #include "Styling/AppStyle.h"
diff --git a/Engine/Plugins/Animation/ControlRig/Source/ControlRigEditor/Private/Tools/ControlRigSnapper.cpp b/Engine/Plugins/Animation/ControlRig/Source/ControlRigEditor/Private/Tools/ControlRigSnapper.cpp
--- a/Engine/Plugins/Animation/ControlRig/Source/ControlRigEditor/Private/Tools/ControlRigSnapper.cpp
+++ b/Engine/Plugins/Animation/ControlRig/Source/ControlRigEditor/Private/Tools/ControlRigSnapper.cpp
@@ -21,7 +21,6 @@
 #include "Sections/MovieScene3DTransformSection.h"
 #include "Tracks/MovieScene3DTransformTrack.h"
 #include "Channels/MovieSceneChannelProxy.h"
-#include "MatineeImportTools.h"
 #include "Tools/ControlRigSnapSettings.h"
 #include "MovieSceneToolsModule.h"
 #include "Components/SkeletalMeshComponent.h"
diff --git a/Engine/Plugins/Cameras/CameraShakePreviewer/Source/CameraShakePreviewer/Private/SCameraShakePreviewer.cpp b/Engine/Plugins/Cameras/CameraShakePreviewer/Source/CameraShakePreviewer/Private/SCameraShakePreviewer.cpp
--- a/Engine/Plugins/Cameras/CameraShakePreviewer/Source/CameraShakePreviewer/Private/SCameraShakePreviewer.cpp
+++ b/Engine/Plugins/Cameras/CameraShakePreviewer/Source/CameraShakePreviewer/Private/SCameraShakePreviewer.cpp
@@ -2,8 +2,6 @@
 
 #include "SCameraShakePreviewer.h"
 #include "Camera/CameraActor.h"
-#include "Camera/CameraAnim.h"
-#include "Camera/CameraAnimInst.h"
 #include "Camera/CameraModifier_CameraShake.h"
 #include "Camera/CameraShakeBase.h"
 #include "Camera/CameraShakeSourceActor.h"
@@ -120,11 +118,6 @@
 
 		PreviewCamera->MergePostProcessSettings(LastPostProcessSettings, LastPostProcessBlendWeights);
 
-		for (UCameraAnimInst* ActiveAnim : ActiveAnims)
-		{
-			UpdateCameraAnimInstance(*ActiveAnim, DeltaTime, Params.ViewInfo);
-		}
-
 		LastLocationModifier = Params.ViewInfo.Location - OriginalPOV.Location;
 		LastRotationModifier = Params.ViewInfo.Rotation - OriginalPOV.Rotation;
 		LastFOVModifier = Params.ViewInfo.FOV - OriginalPOV.FOV;
@@ -142,16 +135,6 @@
 	{
 		Params.AddPostProcessBlend(LastPostProcessSettings[PPIndex], LastPostProcessBlendWeights[PPIndex]);
 	}
-
-	// Clean-up finished camera anims.
-	for (int32 Index = ActiveAnims.Num() - 1; Index >= 0; --Index)
-	{
-		UCameraAnimInst* ActiveAnim(ActiveAnims[Index]);
-		if (ActiveAnim->bFinished)
-		{
-			ActiveAnims.RemoveAt(Index);
-		}
-	}
 }
 
 
@@ -159,38 +142,20 @@
 void FCameraShakePreviewUpdater::AddReferencedObjects(FReferenceCollector& Collector)
 {
 	Collector.AddReferencedObject(PreviewCamera);
-	Collector.AddReferencedObjects(ActiveAnims);
 }
 
 UCameraShakeBase* FCameraShakePreviewUpdater::AddCameraShake(TSubclassOf<UCameraShakeBase> ShakeClass, const FAddCameraShakeParams& Params)
 {
 	FAddCameraShakeParams ActualParams(Params);
-	if (ShakeClass.Get()->IsChildOf<UMatineeCameraShake>())
-	{
-		ActualParams.Initializer = FOnInitializeCameraShake::CreateLambda([this](UCameraShakeBase* ShakeInstance)
-			{
-				UMatineeCameraShake* MatineeShakeInstance = CastChecked<UMatineeCameraShake>(ShakeInstance);
-				MatineeShakeInstance->SetTempCameraAnimActor(GetTempCameraActor());
-			});
-	}
 
 	UCameraShakeBase* ShakeInstance = PreviewCameraShake->AddCameraShake(ShakeClass, ActualParams);
 
-	if (UMatineeCameraShake* MatineeShakeInstance = Cast<UMatineeCameraShake>(ShakeInstance))
-	{
-		if (UCameraAnimInst* CameraAnimInstance = MatineeShakeInstance->AnimInst)
-		{
-			ActiveAnims.Add(CameraAnimInstance);
-		}
-	}
 	return ShakeInstance;
 }
 
 void FCameraShakePreviewUpdater::RemoveAllCameraShakesFromSource(const UCameraShakeSourceComponent* SourceComponent)
 {
 	PreviewCameraShake->RemoveAllCameraShakesFromSource(SourceComponent);
-
-	CleanUpCameraAnimInstances();
 }
 
 void FCameraShakePreviewUpdater::GetActiveCameraShakes(TArray<FActiveCameraShakeInfo>& ActiveCameraShakes) const
@@ -201,8 +166,6 @@
 void FCameraShakePreviewUpdater::RemoveCameraShake(UCameraShakeBase* ShakeInstance)
 {
 	PreviewCameraShake->RemoveCameraShake(ShakeInstance);
-
-	CleanUpCameraAnimInstances();
 }
 
 void FCameraShakePreviewUpdater::RemoveAllCameraShakes()
@@ -210,82 +173,6 @@
 	PreviewCameraShake->RemoveAllCameraShakes();
 }
 
-ACameraActor* FCameraShakePreviewUpdater::GetTempCameraActor()
-{
-	if (!TempCameraActor.IsValid())
-	{
-		UWorld* World = UE::MovieScene::FindCameraShakePreviewerWorld();
-		if (!ensure(World))
-		{
-			return nullptr;
-		}
-
-		FActorSpawnParameters SpawnInfo;
-		SpawnInfo.SpawnCollisionHandlingOverride = ESpawnActorCollisionHandlingMethod::AlwaysSpawn;
-
-		// Don't save this actor into a map.
-		SpawnInfo.ObjectFlags |= RF_Transient;
-		TempCameraActor = World->SpawnActor<ACameraActor>(SpawnInfo);
-		if (TempCameraActor.IsValid())
-		{
-			TempCameraActor.Get()->SetIsTemporarilyHiddenInEditor(true);
-		}
-	}
-
-	check(TempCameraActor.IsValid());
-	return TempCameraActor.Get();
-}
-
-void FCameraShakePreviewUpdater::UpdateCameraAnimInstance(UCameraAnimInst& CameraAnimInstance, float DeltaTime, FMinimalViewInfo& InOutPOV)
-{
-	ACameraActor* Actor = TempCameraActor.Get();
-	if (!ensure(Actor != nullptr))
-	{
-		return;
-	}
-	if (CameraAnimInstance.bFinished)
-	{
-		return;
-	}
-
-	Actor->SetActorLocationAndRotation(FVector::ZeroVector, FRotator::ZeroRotator);
-
-	const ACameraActor* DefaultActor = GetDefault<ACameraActor>();
-	if (DefaultActor)
-	{
-		Actor->GetCameraComponent()->AspectRatio = DefaultActor->GetCameraComponent()->AspectRatio;
-		Actor->GetCameraComponent()->FieldOfView = CameraAnimInstance.CamAnim->BaseFOV;
-		Actor->GetCameraComponent()->PostProcessSettings = CameraAnimInstance.CamAnim->BasePostProcessSettings;
-		Actor->GetCameraComponent()->PostProcessBlendWeight = CameraAnimInstance.CamAnim->BasePostProcessBlendWeight;
-	}
-
-	CameraAnimInstance.AdvanceAnim(DeltaTime, false);
-
-	if (CameraAnimInstance.CurrentBlendWeight > 0.f)
-	{
-		CameraAnimInstance.ApplyToView(InOutPOV);
-
-		if (Actor->GetCameraComponent()->PostProcessBlendWeight > 0.f)
-		{
-			AddPostProcessBlend(
-					Actor->GetCameraComponent()->PostProcessSettings, 
-					Actor->GetCameraComponent()->PostProcessBlendWeight * CameraAnimInstance.CurrentBlendWeight);
-		}
-	}
-}
-
-void FCameraShakePreviewUpdater::CleanUpCameraAnimInstances()
-{
-	for (int32 Index = ActiveAnims.Num() - 1; Index >= 0; --Index)
-	{
-		UCameraAnimInst* ActiveAnim(ActiveAnims[Index]);
-		if (!ActiveAnim || ActiveAnim->bFinished)
-		{
-			ActiveAnims.RemoveAt(Index);
-		}
-	}
-}
-
 void FCameraShakePreviewUpdater::AddPostProcessBlend(const FPostProcessSettings& Settings, float Weight)
 {
 	check(LastPostProcessSettings.Num() == LastPostProcessBlendWeights.Num());
diff --git a/Engine/Plugins/Cameras/CameraShakePreviewer/Source/CameraShakePreviewer/Private/SCameraShakePreviewer.h b/Engine/Plugins/Cameras/CameraShakePreviewer/Source/CameraShakePreviewer/Private/SCameraShakePreviewer.h
--- a/Engine/Plugins/Cameras/CameraShakePreviewer/Source/CameraShakePreviewer/Private/SCameraShakePreviewer.h
+++ b/Engine/Plugins/Cameras/CameraShakePreviewer/Source/CameraShakePreviewer/Private/SCameraShakePreviewer.h
@@ -73,10 +73,6 @@
 	void ModifyCamera(FEditorViewportViewModifierParams& Params);
 
 private:
-	ACameraActor* GetTempCameraActor();
-	void UpdateCameraAnimInstance(UCameraAnimInst& CameraAnimInstance, float DeltaTime, FMinimalViewInfo& InOutPOV);
-	void CleanUpCameraAnimInstances();
-
 	void AddPostProcessBlend(const FPostProcessSettings& Settings, float Weight);
 
 	void OnObjectsReplaced(const TMap<UObject*, UObject*>& ReplacementMap);
@@ -85,10 +81,6 @@
 	APreviewPlayerCameraManager* PreviewCamera;
 	UCameraModifier_CameraShake* PreviewCameraShake;
 
-	/** Hidden camera actor and active camera anims for Matinee shakes specifically */
-	TWeakObjectPtr<ACameraActor> TempCameraActor;
-	TArray<UCameraAnimInst*> ActiveAnims;
-
 	TOptional<float> LastDeltaTime;
 
 	FVector LastLocationModifier;
diff --git a/Engine/Plugins/Cameras/GameplayCameras/Source/GameplayCameras/Private/MatineeCameraShake.cpp b/Engine/Plugins/Cameras/GameplayCameras/Source/GameplayCameras/Private/MatineeCameraShake.cpp
--- a/Engine/Plugins/Cameras/GameplayCameras/Source/GameplayCameras/Private/MatineeCameraShake.cpp
+++ b/Engine/Plugins/Cameras/GameplayCameras/Source/GameplayCameras/Private/MatineeCameraShake.cpp
@@ -3,12 +3,9 @@
 #include "MatineeCameraShake.h"
 #include "SequenceCameraShake.h"
 #include "Camera/CameraActor.h"
-#include "Camera/CameraAnim.h"
-#include "Camera/CameraAnimInst.h"
 #include "Camera/CameraComponent.h"
 #include "Camera/PlayerCameraManager.h"
 #include "Engine/Engine.h"
-#include "Evaluation/MovieSceneCameraAnimTemplate.h"
 #include "IXRTrackingSystem.h" // for IsHeadTrackingAllowed()
 
 DEFINE_LOG_CATEGORY_STATIC(LogMatineeCameraShake, Warning, All);
@@ -86,21 +83,6 @@ void UMatineeCameraShake::DoStopShake(bool bImmediately)
 
 	if (bImmediately)
 	{
-		// stop cam anim if playing
-		if (AnimInst && !AnimInst->bFinished)
-		{
-			if (CameraOwner)
-			{
-				CameraOwner->StopCameraAnimInst(AnimInst, true);
-			}
-			else
-			{
-				AnimInst->Stop(true);
-			}
-		}
-
-		AnimInst = nullptr;
-
 		// stop oscillation
 		OscillatorTimeRemaining = 0.f;
 	}
@@ -115,19 +97,6 @@ void UMatineeCameraShake::DoStopShake(bool bImmediately)
 		{
 			OscillatorTimeRemaining = OscillationBlendOutTime;
 		}
-
-		if (AnimInst && !AnimInst->bFinished)
-		{
-			if (CameraOwner)
-			{
-				CameraOwner->StopCameraAnimInst(AnimInst, false);
-			}
-			else
-			{
-				// playing without a cameramanager, stop it ourselves
-				AnimInst->Stop(false);
-			}
-		}
 	}
 	
 	if (SequenceShakePattern)
@@ -200,52 +169,7 @@ void UMatineeCameraShake::DoStartShake(const FCameraShakeStartParams& Params)
 
 	// init cameraanim shakes
 	APlayerCameraManager* CameraOwner = GetCameraManager();
-	if (Anim != nullptr)
-	{
-		if (AnimInst)
-		{
-			float const Duration = bRandomAnimSegment ? RandomAnimSegmentDuration : 0.f;
-			float const FinalAnimScale = ShakeScale * AnimScale;
-			AnimInst->Update(AnimPlayRate, FinalAnimScale, AnimBlendInTime, AnimBlendOutTime, Duration);
-		}
-		else
-		{
-			bool bLoop = false;
-			bool bRandomStart = false;
-			float Duration = 0.f;
-			if (bRandomAnimSegment)
-			{
-				bLoop = true;
-				bRandomStart = true;
-				Duration = RandomAnimSegmentDuration;
-			}
-
-			float const FinalAnimScale = ShakeScale * AnimScale;
-			if (FinalAnimScale > 0.f)
-			{
-				ECameraShakePlaySpace AnimPlaySpace = GetPlaySpace();
-				FRotator UserPlaySpaceRot = GetUserPlaySpaceMatrix().Rotator();
-
-				if (CameraOwner)
-				{
-					AnimInst = CameraOwner->PlayCameraAnim(Anim, AnimPlayRate, FinalAnimScale, AnimBlendInTime, AnimBlendOutTime, bLoop, bRandomStart, Duration, AnimPlaySpace, UserPlaySpaceRot);
-				}
-				else
-				{
-					// allocate our own instance and start it
-					AnimInst = NewObject<UCameraAnimInst>(this);
-					if (AnimInst)
-					{
-						// note: we don't have a temp camera actor necessary for evaluating a camera anim.
-						// caller is responsible in this case for providing one by calling SetTempCameraAnimActor() on the shake instance before playing the shake
-						AnimInst->Play(Anim, TempCameraActorForCameraAnims.Get(), AnimPlayRate, FinalAnimScale, AnimBlendInTime, AnimBlendOutTime, bLoop, bRandomStart, Duration);
-						AnimInst->SetPlaySpace(AnimPlaySpace, UserPlaySpaceRot);
-					}
-				}
-			}
-		}
-	}
-	else if (AnimSequence != nullptr)
+	if (AnimSequence != nullptr)
 	{
 		if (SequenceShakePattern == nullptr)
 		{
@@ -280,12 +204,6 @@ void UMatineeCameraShake::DoUpdateShake(const FCameraShakeUpdateParams& Params,
 	const float DeltaTime = Params.DeltaTime;
 	const float BaseShakeScale = Params.GetTotalScale();
 
-	// update anims with any desired scaling
-	if (AnimInst)
-	{
-		AnimInst->TransientScaleModifier *= BaseShakeScale;
-	}
-
 	// update oscillation times... only decrease the time remaining if we're not infinite
 	if (OscillatorTimeRemaining > 0.f)
 	{
@@ -370,8 +288,7 @@ void UMatineeCameraShake::DoUpdateShake(const FCameraShakeUpdateParams& Params,
 
 			// Don't allow shake to flip pitch past vertical, if not using a headset (where we can't limit the camera locked to your head).
 			APlayerCameraManager* CameraOwner = GetCameraManager();
-			AActor * WorldActor = (CameraOwner ? CameraOwner : TempCameraActorForCameraAnims.Get());
-			UWorld * World = (WorldActor ? WorldActor->GetWorld() : nullptr);
+			UWorld * World = (CameraOwner ? CameraOwner->GetWorld() : nullptr);
 			if (!GEngine->XRSystem.IsValid() ||
 				!(World != nullptr ? GEngine->XRSystem->IsHeadTrackingAllowedForWorld(*World) : GEngine->XRSystem->IsHeadTrackingAllowed()))
 			{
@@ -479,27 +396,11 @@ void UMatineeCameraShake::DoScrubShake(const FCameraShakeScrubParams& Params, FC
 	DoUpdateShake(UpdateParams, OutResult);
 
 	check(EnumHasAnyFlags(OutResult.Flags, ECameraShakeUpdateResultFlags::ApplyAsAbsolute));
-
-	if (AnimInst)
-	{
-		FMinimalViewInfo AnimPOV(Params.POV);
-		AnimPOV.Location = OutResult.Location;
-		AnimPOV.Rotation = OutResult.Rotation;
-		AnimPOV.FOV = OutResult.FOV;
-
-		AnimInst->SetCurrentTime(NewTime);
-		AnimInst->ApplyToView(AnimPOV);
-
-		OutResult.Location = AnimPOV.Location;
-		OutResult.Rotation = AnimPOV.Rotation;
-		OutResult.FOV = AnimPOV.FOV;
-	}
 }
 
 bool UMatineeCameraShake::DoGetIsFinished() const
 {
 	return ((OscillatorTimeRemaining <= 0.f) &&									// oscillator is finished
-		((AnimInst == nullptr) || AnimInst->bFinished) &&						// anim is finished
 		((SequenceShakePattern == nullptr) ||									// other anim is finished
 			SequenceShakeState.GetElapsedTime() >= SequenceShakeState.GetDuration()) &&
 		ReceiveIsFinished()														// BP thinks it's finished
@@ -557,8 +458,22 @@ void UMatineeCameraShakePattern::GetShakePatternInfoImpl(FCameraShakeInfo& OutIn
 {
 	// We will manage our own duration, but let's give a hint about how long we are for editor purposes.
 	UMatineeCameraShake* Shake = GetShakeInstance<UMatineeCameraShake>();
-	const float Duration = FMath::Max(Shake->OscillationDuration, Shake->Anim ? Shake->Anim->AnimLength : 0.f);
-	OutInfo.Duration = FCameraShakeDuration::Custom(Duration);
+	
+	if (Shake->AnimSequence)
+	{
+		float PlaybackSeconds = 0.f;
+		if (UMovieScene* MovieScene = Shake->AnimSequence->GetMovieScene())
+		{
+			const FFrameNumber PlaybackFrames = MovieScene->GetPlaybackRange().Size<FFrameNumber>();
+			PlaybackSeconds = MovieScene->GetTickResolution().AsSeconds(PlaybackFrames);
+		}
+		const float Duration = FMath::Max(Shake->OscillationDuration, PlaybackSeconds);
+		OutInfo.Duration = FCameraShakeDuration::Custom(Duration);
+	}
+	else
+	{
+		OutInfo.Duration = FCameraShakeDuration::Custom(Shake->OscillationDuration);
+	}
 }
 
 void UMatineeCameraShakePattern::StopShakePatternImpl(const FCameraShakeStopParams& Params)
@@ -597,62 +512,3 @@ void UMatineeCameraShakePattern::TeardownShakePatternImpl()
 	Shake->DoTeardownShake();
 }
 
-UMovieSceneMatineeCameraShakeEvaluator::UMovieSceneMatineeCameraShakeEvaluator(const FObjectInitializer& ObjInit)
-	: Super(ObjInit)
-{
-	if (HasAnyFlags(EObjectFlags::RF_ClassDefaultObject))
-	{
-		FMovieSceneCameraShakeEvaluatorRegistry::RegisterShakeEvaluatorBuilder(
-			FMovieSceneBuildShakeEvaluator::CreateStatic(&UMovieSceneMatineeCameraShakeEvaluator::BuildMatineeShakeEvaluator));
-	}
-}
-
-UMovieSceneCameraShakeEvaluator* UMovieSceneMatineeCameraShakeEvaluator::BuildMatineeShakeEvaluator(UCameraShakeBase* ShakeInstance)
-{
-	if (UMatineeCameraShake* MatineeShakeInstance = Cast<UMatineeCameraShake>(ShakeInstance))
-	{
-		return NewObject<UMovieSceneMatineeCameraShakeEvaluator>();
-	}
-	return nullptr;
-}
-
-bool UMovieSceneMatineeCameraShakeEvaluator::Setup(const FMovieSceneEvaluationOperand& Operand, FPersistentEvaluationData& PersistentData, IMovieScenePlayer& Player, UCameraShakeBase* ShakeInstance)
-{
-	UMatineeCameraShake* MatineeShakeInstance = CastChecked<UMatineeCameraShake>(ShakeInstance);
-
-	// We use the global temp actor from the shared data (shared across all additive camera effects for this operand)
-	ACameraActor* TempCameraActor = FMovieSceneMatineeCameraData::Get(Operand, PersistentData).GetTempCameraActor(Player);
-	MatineeShakeInstance->SetTempCameraAnimActor(TempCameraActor);
-
-	if (MatineeShakeInstance->AnimInst)
-	{
-		MatineeShakeInstance->AnimInst->SetStopAutomatically(false);
-	}
-
-	return true;
-}
-
-bool UMovieSceneMatineeCameraShakeEvaluator::Evaluate(const FMovieSceneContext& Context, const FMovieSceneEvaluationOperand& Operand, FPersistentEvaluationData& PersistentData, IMovieScenePlayer& Player, UCameraShakeBase* ShakeInstance)
-{
-	UMatineeCameraShake* MatineeShakeInstance = CastChecked<UMatineeCameraShake>(ShakeInstance);
-
-	FMovieSceneMatineeCameraData& MatineeSharedData = FMovieSceneMatineeCameraData::Get(Operand, PersistentData);
-	ACameraActor* TempCameraActor = MatineeSharedData.GetTempCameraActor(Player);
-	
-	// prepare temp camera actor by resetting it
-	TempCameraActor->SetActorLocationAndRotation(FVector::ZeroVector, FRotator::ZeroRotator);
-
-	ACameraActor const* const DefaultCamActor = GetDefault<ACameraActor>();
-	if (DefaultCamActor)
-	{
-		TempCameraActor->GetCameraComponent()->AspectRatio = DefaultCamActor->GetCameraComponent()->AspectRatio;
-
-		UCameraAnim* CamAnim = MatineeShakeInstance->AnimInst ? MatineeShakeInstance->AnimInst->CamAnim : nullptr;
-
-		TempCameraActor->GetCameraComponent()->PostProcessSettings = CamAnim ? CamAnim->BasePostProcessSettings : FPostProcessSettings();
-		TempCameraActor->GetCameraComponent()->PostProcessBlendWeight = CamAnim ? MatineeShakeInstance->AnimInst->CamAnim->BasePostProcessBlendWeight : 0.f;
-	}
-
-	return true;
-}
-
diff --git a/Engine/Plugins/Cameras/GameplayCameras/Source/GameplayCameras/Public/MatineeCameraShake.h b/Engine/Plugins/Cameras/GameplayCameras/Source/GameplayCameras/Public/MatineeCameraShake.h
--- a/Engine/Plugins/Cameras/GameplayCameras/Source/GameplayCameras/Public/MatineeCameraShake.h
+++ b/Engine/Plugins/Cameras/GameplayCameras/Source/GameplayCameras/Public/MatineeCameraShake.h
@@ -154,7 +154,7 @@ public:
 	struct FFOscillator FOVOscillation;
 
 	/************************************************************
-	 * Parameters for defining CameraAnim-driven camera shakes
+	 * Parameters for defining sequencer-driven camera shakes
 	 ************************************************************/
 
 	/** Scalar defining how fast to play the anim. */
@@ -177,17 +177,13 @@ public:
 	UPROPERTY(EditAnywhere, Category = AnimShake, meta = (ClampMin = "0.0", editcondition = "bRandomAnimSegment"))
 	float RandomAnimSegmentDuration;
 
-	/** Source camera animation to play. Can be null. */
-	UPROPERTY(EditAnywhere, Category = AnimShake)
-	TObjectPtr<class UCameraAnim> Anim;
-
-	/** Source camera animation sequence to play. Can be null, but can't have both Anim and AnimSequence. */
+	/** Source camera animation sequence to play. Can be null. */
 	UPROPERTY(EditAnywhere, Category = AnimShake)
 	TObjectPtr<class UCameraAnimationSequence> AnimSequence;
 
 	/**
-	* If true, play a random snippet of the animation of length Duration.  Implies bLoop and bRandomStartTime = true for the CameraAnim.
-	* If false, play the full anim once, non-looped. Useful for getting variety out of a single looped CameraAnim asset.
+	* If true, play a random snippet of the animation of length Duration.  Implies bLoop and bRandomStartTime = true for the AnimSequence.
+	* If false, play the full anim once, non-looped. Useful for getting variety out of a single looped AnimSequence asset.
 	*/
 	UPROPERTY(EditAnywhere, Category = AnimShake)
 	uint32 bRandomAnimSegment : 1;
@@ -198,10 +194,6 @@ public:
 	UPROPERTY(transient, BlueprintReadOnly, Category = CameraShake)
 	float OscillatorTimeRemaining;
 
-	/** The playing instance of the CameraAnim-based shake, if any. */
-	UPROPERTY(transient, BlueprintReadOnly, Category = CameraShake)
-	TObjectPtr<class UCameraAnimInst> AnimInst;
-
 public:
 
 	// Blueprint API
@@ -246,12 +238,9 @@ public:
 	/** Returns true if this camera shake will loop forever */
 	bool IsLooping() const;
 
-	/** Sets current playback time and applies the shake (both oscillation and cameraanim) to the given POV. */
+	/** Sets current playback time and applies the shake (both oscillation and camera animation sequence) to the given POV. */
 	UE_DEPRECATED(4.27, "SetCurrentTimeAndApplyShake is deprecated, please use ScrubAndApplyCameraShake")
 	void SetCurrentTimeAndApplyShake(float NewTime, FMinimalViewInfo& POV);
-	
-	/** Sets actor for playing camera anims */
-	void SetTempCameraAnimActor(AActor* Actor) { TempCameraActorForCameraAnims = Actor; }
 
 private:
 
@@ -282,9 +271,6 @@ protected:
 	/** Initial offset (could have been assigned at random). */
 	float InitialFOVSinOffset;
 
-	/** Temp actor to use for playing camera anims. Used when playing a camera anim in non-gameplay context, e.g. in the editor */
-	TWeakObjectPtr<AActor> TempCameraActorForCameraAnims;
-
 	/** Sequence shake pattern for when using a sequence instead of a camera anim */
 	UPROPERTY(Instanced)
 	TObjectPtr<class USequenceCameraShakePattern> SequenceShakePattern;
@@ -330,20 +316,6 @@ private:
 UE_DEPRECATED(4.26, "Please use UMatineeCameraShake")
 typedef UMatineeCameraShake UCameraShake;
 
-/** Custom sequencer evaluation code for Matinee camera shakes */
-UCLASS()
-class UMovieSceneMatineeCameraShakeEvaluator : public UMovieSceneCameraShakeEvaluator
-{
-	GENERATED_BODY()
-
-	UMovieSceneMatineeCameraShakeEvaluator(const FObjectInitializer& ObjInit);
-
-	static UMovieSceneCameraShakeEvaluator* BuildMatineeShakeEvaluator(UCameraShakeBase* ShakeInstance);
-
-	virtual bool Setup(const FMovieSceneEvaluationOperand& Operand, FPersistentEvaluationData& PersistentData, IMovieScenePlayer& Player, UCameraShakeBase* ShakeInstance) override;
-	virtual bool Evaluate(const FMovieSceneContext& Context, const FMovieSceneEvaluationOperand& Operand, FPersistentEvaluationData& PersistentData, IMovieScenePlayer& Player, UCameraShakeBase* ShakeInstance) override;
-};
-
 /**
  * Blueprint function library for autocasting from a base camera shake to a matinee camera shake.
  * This prevents breaking Blueprints now that APlayerCameraManager::StartCameraShake returns the base class.
diff --git a/Engine/Plugins/MovieScene/LevelSequenceEditor/Source/LevelSequenceEditor/Private/LevelSequenceEditorSubsystem.cpp b/Engine/Plugins/MovieScene/LevelSequenceEditor/Source/LevelSequenceEditor/Private/LevelSequenceEditorSubsystem.cpp
--- a/Engine/Plugins/MovieScene/LevelSequenceEditor/Source/LevelSequenceEditor/Private/LevelSequenceEditorSubsystem.cpp
+++ b/Engine/Plugins/MovieScene/LevelSequenceEditor/Source/LevelSequenceEditor/Private/LevelSequenceEditorSubsystem.cpp
@@ -20,7 +20,6 @@
 #include "Sections/MovieScene3DTransformSection.h"
 #include "Tracks/MovieScene3DTransformTrack.h"
 #include "Tracks/MovieScene3DConstraintTrack.h"
-#include "Tracks/MovieSceneCameraAnimTrack.h"
 #include "Tracks/MovieSceneCameraShakeTrack.h"
 
 #include "ActorTreeItem.h"
@@ -839,26 +838,6 @@ void ULevelSequenceEditorSubsystem::BakeTransform(const TArray<FMovieSceneBindin
 			}
 		}
 
-		// Disable or delete any camera anim tracks
-		for (UMovieSceneTrack* Track : FocusedMovieScene->FindTracks(UMovieSceneCameraAnimTrack::StaticClass(), Guid))
-		{
-			if (UMovieSceneCameraAnimTrack* CameraAnimTrack = Cast<UMovieSceneCameraAnimTrack>(Track))
-			{
-				if (bDisableSectionsAfterBaking)
-				{
-					for (UMovieSceneSection* CameraAnimSection : CameraAnimTrack->GetAllSections())
-					{
-						CameraAnimSection->Modify();
-						CameraAnimSection->SetIsActive(false);
-					}
-				}
-				else
-				{
-					FocusedMovieScene->RemoveTrack(*CameraAnimTrack);
-				}
-			}
-		}
-
 		// Disable or delete any camera shake tracks
 		for (UMovieSceneTrack* Track : FocusedMovieScene->FindTracks(UMovieSceneCameraShakeTrack::StaticClass(), Guid))
 		{
diff --git a/Engine/Plugins/MovieScene/TemplateSequence/Source/TemplateSequenceEditor/Private/TrackEditors/TemplateSequenceTrackEditor.cpp b/Engine/Plugins/MovieScene/TemplateSequence/Source/TemplateSequenceEditor/Private/TrackEditors/TemplateSequenceTrackEditor.cpp
--- a/Engine/Plugins/MovieScene/TemplateSequence/Source/TemplateSequenceEditor/Private/TrackEditors/TemplateSequenceTrackEditor.cpp
+++ b/Engine/Plugins/MovieScene/TemplateSequence/Source/TemplateSequenceEditor/Private/TrackEditors/TemplateSequenceTrackEditor.cpp
@@ -4,7 +4,6 @@
 #include "AssetRegistry/AssetData.h"
 #include "AssetRegistry/AssetRegistryModule.h"
 #include "Camera/CameraComponent.h"
-#include "Camera/CameraAnim.h"
 #include "CameraAnimationSequence.h"
 #include "CollectionManagerModule.h"
 #include "CommonMovieSceneTools.h"
@@ -19,12 +18,9 @@
 #include "Misc/ConfigCacheIni.h"
 #include "MovieSceneCommonHelpers.h"
 #include "MovieSceneTimeHelpers.h"
-#include "Sections/MovieSceneCameraAnimSection.h"
 #include "SequencerSectionPainter.h"
 #include "SequencerUtilities.h"
 #include "TemplateSequence.h"
-#include "TrackEditors/CameraAnimTrackEditorHelper.h"
-#include "Tracks/MovieSceneCameraAnimTrack.h"
 #include "Tracks/MovieSceneFloatTrack.h"
 #include "Tracks/MovieScene3DTransformTrack.h"
 #include "Tracks/TemplateSequenceTrack.h"
@@ -63,28 +59,12 @@ bool FTemplateSequenceTrackEditor::SupportsSequence(UMovieSceneSequence* InSeque
 
 void FTemplateSequenceTrackEditor::BuildObjectBindingTrackMenu(FMenuBuilder& MenuBuilder, const TArray<FGuid>& ObjectBindings, const UClass* ObjectClass)
 {
-	bool bIsCameraAnimMenu = true;
-	for (FGuid ObjectBinding : ObjectBindings)
-	{
-		const UCameraComponent* CameraComponent = AcquireCameraComponentFromObjectGuid(ObjectBinding);
-		if (CameraComponent == nullptr)
-		{
-			bIsCameraAnimMenu = false;
-			break;
-		}
-	}
-
-	FText SubMenuEntryText = LOCTEXT("AddTemplateSequence", "Template Sequence");
-	FText SubMenuEntryTooltip = LOCTEXT("AddTemplateSequenceTooltip", "Adds a track that can play a template sequence asset using the parent binding.");
-	if (bIsCameraAnimMenu)
-	{
-		SubMenuEntryText = LOCTEXT("AddCameraAnimationSequence", "Camera Animation");
-		SubMenuEntryTooltip = LOCTEXT("AddCameraAnimationSequenceTooltip", "Adds a camera animation template sequence on the parent binding.");
-	}
+	const FText SubMenuEntryText = LOCTEXT("AddTemplateSequence", "Template Sequence");
+	const FText SubMenuEntryTooltip = LOCTEXT("AddTemplateSequenceTooltip", "Adds a track that can play a template sequence asset using the parent binding.");
 	
 	MenuBuilder.AddSubMenu(
 		SubMenuEntryText, SubMenuEntryTooltip,
-		FNewMenuDelegate::CreateRaw(this, &FTemplateSequenceTrackEditor::AddTemplateSequenceAssetSubMenu, ObjectBindings, ObjectClass, bIsCameraAnimMenu)
+		FNewMenuDelegate::CreateRaw(this, &FTemplateSequenceTrackEditor::AddTemplateSequenceAssetSubMenu, ObjectBindings, ObjectClass)
 	);
 }
 
@@ -94,9 +74,6 @@ TSharedPtr<SWidget> FTemplateSequenceTrackEditor::BuildOutlinerEditWidget(const
 
 	if (ObjectClass != nullptr)
 	{
-		const UCameraComponent* CameraComponent = AcquireCameraComponentFromObjectGuid(ObjectBinding);
-		const bool bIsCameraAnimMenu = (CameraComponent != nullptr);
-
 		return SNew(SHorizontalBox)
 			+ SHorizontalBox::Slot()
 			.AutoWidth()
@@ -104,7 +81,7 @@ TSharedPtr<SWidget> FTemplateSequenceTrackEditor::BuildOutlinerEditWidget(const
 			[
 				FSequencerUtilities::MakeAddButton(
 					LOCTEXT("TemplateSequenceAddButton", "Template Sequence"),
-					FOnGetContent::CreateSP(this, &FTemplateSequenceTrackEditor::BuildTemplateSequenceAssetSubMenu, ObjectBinding, ObjectClass, bIsCameraAnimMenu),
+					FOnGetContent::CreateSP(this, &FTemplateSequenceTrackEditor::BuildTemplateSequenceAssetSubMenu, ObjectBinding, ObjectClass),
 					Params.NodeIsHovered, GetSequencer())
 			];
 	}
@@ -128,11 +105,10 @@ public:
 	{}
 	SLATE_END_ARGS();
 
-	void Construct(const FArguments& Args, TSharedPtr<FTemplateSequenceTrackEditor> InTrackEditor, TArray<FGuid> ObjectBindings, const UClass* InBaseClass, const UClass* InLegacyBaseClass)
+	void Construct(const FArguments& Args, TSharedPtr<FTemplateSequenceTrackEditor> InTrackEditor, TArray<FGuid> ObjectBindings, const UClass* InBaseClass)
 	{
 		check(InBaseClass != nullptr);
 		BaseClass = InBaseClass;
-		LegacyBaseClass = InLegacyBaseClass;
 
 		TrackEditor = InTrackEditor;
 
@@ -162,52 +138,39 @@ public:
 			AssetPickerConfig.Filter.ClassPaths.Add(UTemplateSequence::StaticClass()->GetClassPathName());
 			AssetPickerConfig.SaveSettingsName = TEXT("SequencerAssetPicker");
 			AssetPickerConfig.AdditionalReferencingAssets.Add(FAssetData(Sequence));
-			if (LegacyBaseClass != nullptr)
-			{
-				AssetPickerConfig.Filter.ClassPaths.Add(LegacyBaseClass->GetClassPathName());
-			}
 
 			AssetPickerConfig.OnShouldFilterAsset = FOnShouldFilterAsset::CreateLambda(
 				[this, BaseClassNames](const FAssetData& AssetData) -> bool
 				{
-					if (LegacyBaseClass == nullptr || AssetData.AssetClassPath != LegacyBaseClass->GetClassPathName())
+					FString FoundBoundActorClass;
+					const FAssetDataTagMapSharedView::FFindTagResult FoundBoundActorClassTag = AssetData.TagsAndValues.FindTag("BoundActorClass");
+					if (FoundBoundActorClassTag.IsSet())
 					{
-						FString FoundBoundActorClass;
-						const FAssetDataTagMapSharedView::FFindTagResult FoundBoundActorClassTag = AssetData.TagsAndValues.FindTag("BoundActorClass");
-						if (FoundBoundActorClassTag.IsSet())
+						FString TagValue = FoundBoundActorClassTag.GetValue();
+						if (FPackageName::IsShortPackageName(TagValue))
 						{
-							FString TagValue = FoundBoundActorClassTag.GetValue();
-							if (FPackageName::IsShortPackageName(TagValue))
-							{
-								// Replace the short name with the full path name if possible.
-								FTopLevelAssetPath BoundActorClassPath = UClass::TryConvertShortTypeNameToPathName<UStruct>(TagValue, ELogVerbosity::Warning, TEXT("STemplateSequenceAssetSubMenu"));
-								if (!BoundActorClassPath.IsNull())
-								{
-									FoundBoundActorClass = BoundActorClassPath.ToString();
-								}
-							}
-							else
+							// Replace the short name with the full path name if possible.
+							FTopLevelAssetPath BoundActorClassPath = UClass::TryConvertShortTypeNameToPathName<UStruct>(TagValue, ELogVerbosity::Warning, TEXT("STemplateSequenceAssetSubMenu"));
+							if (!BoundActorClassPath.IsNull())
 							{
-								FoundBoundActorClass = TagValue;
+								FoundBoundActorClass = BoundActorClassPath.ToString();
 							}
 						}
-						if (!FoundBoundActorClass.IsEmpty())
-						{
-							// Filter this out if it's got an incompatible bound actor class.
-							return !BaseClassNames.Contains(FoundBoundActorClass);
-						}
 						else
 						{
-							// Old asset, hasn't been saved since we added the bound actor class in the tags.
-							++NumAssetsRequiringSave;
-							// Don't filter if we're showing old assets, do filter if we only want compatible assets.
-							return bShowingHiddenAssets ? false : true;
+							FoundBoundActorClass = TagValue;
 						}
 					}
+					if (!FoundBoundActorClass.IsEmpty())
+					{
+						// Filter this out if it's got an incompatible bound actor class.
+						return !BaseClassNames.Contains(FoundBoundActorClass);
+					}
 					else
 					{
-						// Legacy asset (e.g. UCameraAnim asset that has yet to be upgraded to a template sequence).
-						++NumLegacyAssets;
+						// Old asset, hasn't been saved since we added the bound actor class in the tags.
+						++NumAssetsRequiringSave;
+						// Don't filter if we're showing old assets, do filter if we only want compatible assets.
 						return bShowingHiddenAssets ? false : true;
 					}
 				});
@@ -242,7 +205,7 @@ public:
 						.Text_Lambda([this]()
 							{
 								return !bShowingHiddenAssets ? 
-									FText::Format(LOCTEXT("OutdatedAssetsWarning", "Hiding {0} outdated assets"), FText::AsNumber(NumAssetsRequiringSave + NumLegacyAssets)) :
+									FText::Format(LOCTEXT("OutdatedAssetsWarning", "Hiding {0} outdated assets"), FText::AsNumber(NumAssetsRequiringSave)) :
 									LOCTEXT("OutdatedAssetsMessage", "Showing outdated assets");
 							})
 					]
@@ -270,7 +233,7 @@ public:
 
 	EVisibility GetBottomRowVisibility() const
 	{
-		return (NumAssetsRequiringSave + NumLegacyAssets) > 0 ? EVisibility::Visible : EVisibility::Collapsed;
+		return (NumAssetsRequiringSave) > 0 ? EVisibility::Visible : EVisibility::Collapsed;
 	}
 
 	FReply OnShowHiddenAssets()
@@ -278,7 +241,6 @@ public:
 		bShowingHiddenAssets = !bShowingHiddenAssets;
 		GConfig->SetBool(TEXT("TemplateSequence"), TEXT("ShowOutdatedAssetsInCameraAnimationTrackEditor"), bShowingHiddenAssets, GEditorPerProjectIni);
 
-		NumLegacyAssets = 0;
 		NumAssetsRequiringSave = 0;
 		RefreshAssetViewDelegate.ExecuteIfBound(true);
 		return FReply::Handled();
@@ -290,30 +252,27 @@ private:
 	const UClass* LegacyBaseClass;
 
 	TSharedPtr<SWidget> AssetPicker;
-	uint32 NumLegacyAssets = 0;
 	uint32 NumAssetsRequiringSave = 0;
 	bool bShowingHiddenAssets = false;
 	FRefreshAssetViewDelegate RefreshAssetViewDelegate;
 };
 
-void FTemplateSequenceTrackEditor::AddTemplateSequenceAssetSubMenu(FMenuBuilder& MenuBuilder, TArray<FGuid> ObjectBindings, const UClass* RootBindingClass, bool bIsCameraAnimMenu)
+void FTemplateSequenceTrackEditor::AddTemplateSequenceAssetSubMenu(FMenuBuilder& MenuBuilder, TArray<FGuid> ObjectBindings, const UClass* RootBindingClass)
 {
 	if (ObjectBindings.Num() > 0 && RootBindingClass != nullptr)
 	{
-		const UClass* LegacyAssetClass = bIsCameraAnimMenu ? UCameraAnim::StaticClass() : nullptr;
-		TSharedPtr<STemplateSequenceAssetSubMenu> MenuEntry = SNew(STemplateSequenceAssetSubMenu, SharedThis(this), ObjectBindings, RootBindingClass, LegacyAssetClass);
+		TSharedPtr<STemplateSequenceAssetSubMenu> MenuEntry = SNew(STemplateSequenceAssetSubMenu, SharedThis(this), ObjectBindings, RootBindingClass);
 		MenuBuilder.AddWidget(MenuEntry.ToSharedRef(), FText::GetEmpty(), true);
 	}
 }
 
-TSharedRef<SWidget> FTemplateSequenceTrackEditor::BuildTemplateSequenceAssetSubMenu(FGuid ObjectBinding, const UClass* RootBindingClass, bool bIsCameraAnimMenu)
+TSharedRef<SWidget> FTemplateSequenceTrackEditor::BuildTemplateSequenceAssetSubMenu(FGuid ObjectBinding, const UClass* RootBindingClass)
 {
 	check(RootBindingClass != nullptr);
 	
 	TArray<FGuid> ObjectBindings;
 	ObjectBindings.Add(ObjectBinding);
-	const UClass* LegacyAssetClass = bIsCameraAnimMenu ? UCameraAnim::StaticClass() : nullptr;
-	return SNew(STemplateSequenceAssetSubMenu, SharedThis(this), ObjectBindings, RootBindingClass, LegacyAssetClass);
+	return SNew(STemplateSequenceAssetSubMenu, SharedThis(this), ObjectBindings, RootBindingClass);
 }
 
 void FTemplateSequenceTrackEditor::OnTemplateSequenceAssetSelected(const FAssetData& AssetData, TArray<FGuid> ObjectBindings)
@@ -332,21 +291,6 @@ void FTemplateSequenceTrackEditor::OnTemplateSequenceAssetSelected(const FAssetD
 
 		AnimatablePropertyChanged(FOnKeyProperty::CreateRaw(this, &FTemplateSequenceTrackEditor::AddKeyInternal, ObjectBindings, SelectedSequence));
 	}
-	else if (UCameraAnim* SelectedCameraAnim = Cast<UCameraAnim>(AssetData.GetAsset()))
-	{
-		TArray<TWeakObjectPtr<>> OutObjects;
-		for (FGuid ObjectBinding : ObjectBindings)
-		{
-			for (TWeakObjectPtr<> Object : GetSequencer()->FindObjectsInCurrentSequence(ObjectBinding))
-			{
-				OutObjects.Add(Object);
-			}
-		}
-		
-		const FScopedTransaction Transaction(LOCTEXT("AddCameraAnim_Transaction", "Add Camera Anim"));
-
-		AnimatablePropertyChanged(FOnKeyProperty::CreateRaw(this, &FTemplateSequenceTrackEditor::AddLegacyCameraAnimKeyInternal, OutObjects, SelectedCameraAnim));
-	}
 }
 
 void FTemplateSequenceTrackEditor::OnTemplateSequenceAssetEnterPressed(const TArray<FAssetData>& AssetData, TArray<FGuid> ObjectBindings)
@@ -431,11 +375,6 @@ FKeyPropertyResult FTemplateSequenceTrackEditor::AddKeyInternal(FFrameNumber Key
 	return KeyPropertyResult;
 }
 
-FKeyPropertyResult FTemplateSequenceTrackEditor::AddLegacyCameraAnimKeyInternal(FFrameNumber KeyTime, const TArray<TWeakObjectPtr<UObject>> Objects, UCameraAnim* CameraAnim)
-{
-	return FCameraAnimTrackEditorHelper::AddCameraAnimKey(*this, KeyTime, Objects, CameraAnim);
-}
-
 bool FTemplateSequenceTrackEditor::CanAddSubSequence(const UMovieSceneSequence& Sequence) const
 {
 	UMovieSceneSequence* FocusedSequence = GetSequencer()->GetFocusedMovieSceneSequence();
diff --git a/Engine/Plugins/MovieScene/TemplateSequence/Source/TemplateSequenceEditor/Private/TrackEditors/TemplateSequenceTrackEditor.h b/Engine/Plugins/MovieScene/TemplateSequence/Source/TemplateSequenceEditor/Private/TrackEditors/TemplateSequenceTrackEditor.h
--- a/Engine/Plugins/MovieScene/TemplateSequence/Source/TemplateSequenceEditor/Private/TrackEditors/TemplateSequenceTrackEditor.h
+++ b/Engine/Plugins/MovieScene/TemplateSequence/Source/TemplateSequenceEditor/Private/TrackEditors/TemplateSequenceTrackEditor.h
@@ -11,7 +11,6 @@
 
 struct FAssetData;
 class FMenuBuilder;
-class UCameraAnim;
 class UCameraComponent;
 class UTemplateSequence;
 
@@ -31,15 +30,14 @@ public:
 	virtual TSharedRef<ISequencerSection> MakeSectionInterface(UMovieSceneSection& SectionObject, UMovieSceneTrack& Track, FGuid ObjectBinding) override;
 
 private:
-	void AddTemplateSequenceAssetSubMenu(FMenuBuilder& MenuBuilder, TArray<FGuid> ObjectBindings, const UClass* RootBindingClass, bool bIsCameraAnimMenu);
-	TSharedRef<SWidget> BuildTemplateSequenceAssetSubMenu(FGuid ObjectBinding, const UClass* RootBindingClass, bool bIsCameraAnimMenu);
+	void AddTemplateSequenceAssetSubMenu(FMenuBuilder& MenuBuilder, TArray<FGuid> ObjectBindings, const UClass* RootBindingClass);
+	TSharedRef<SWidget> BuildTemplateSequenceAssetSubMenu(FGuid ObjectBinding, const UClass* RootBindingClass);
 
 	void OnTemplateSequenceAssetSelected(const FAssetData& AssetData, TArray<FGuid> ObjectBindings);
 	void OnTemplateSequenceAssetEnterPressed(const TArray<FAssetData>& AssetData, TArray<FGuid> ObjectBindings);
 
 	FKeyPropertyResult AddKeyInternal(FFrameNumber KeyTime, FGuid ObjectBinding, UTemplateSequence* TemplateSequence);
 	FKeyPropertyResult AddKeyInternal(FFrameNumber KeyTime, TArray<FGuid> ObjectBindings, UTemplateSequence* TemplateSequence);
-	FKeyPropertyResult AddLegacyCameraAnimKeyInternal(FFrameNumber KeyTime, const TArray<TWeakObjectPtr<UObject>> Objects, UCameraAnim* CameraAnim);
 
 	bool CanAddSubSequence(const UMovieSceneSequence& Sequence) const;
 	const UClass* AcquireObjectClassFromObjectGuid(const FGuid& Guid);
diff --git a/Engine/Plugins/Runtime/GameplayAbilities/Source/GameplayAbilities/Private/GameplayCueNotifyTypes.cpp b/Engine/Plugins/Runtime/GameplayAbilities/Source/GameplayAbilities/Private/GameplayCueNotifyTypes.cpp
--- a/Engine/Plugins/Runtime/GameplayAbilities/Source/GameplayAbilities/Private/GameplayCueNotifyTypes.cpp
+++ b/Engine/Plugins/Runtime/GameplayAbilities/Source/GameplayAbilities/Private/GameplayCueNotifyTypes.cpp
@@ -9,8 +9,6 @@
 #include "Components/AudioComponent.h"
 #include "Sound/SoundBase.h"
 #include "Camera/CameraShakeBase.h"
-#include "Camera/CameraAnim.h"
-#include "Camera/CameraAnimInst.h"
 #include "Components/ForceFeedbackComponent.h"
 #include "Materials/MaterialInterface.h"
 #include "Components/DecalComponent.h"
@@ -657,70 +655,6 @@ void FGameplayCueNotify_CameraShakeInfo::ValidateBurstAssets(UObject* Containing
 }
 
 
-//////////////////////////////////////////////////////////////////////////
-// FGameplayCueNotify_CameraAnimInfo
-//////////////////////////////////////////////////////////////////////////
-FGameplayCueNotify_CameraAnimInfo::FGameplayCueNotify_CameraAnimInfo()
-{
-	CameraAnim = nullptr;
-	AnimScale = 1.0f;
-	PlayRate = 1.0f;
-
-	BlendInTime = 0.0f;
-	BlendOutTime = 0.0f;
-
-	PlaySpace = EGameplayCueNotify_EffectPlaySpace::CameraSpace;
-
-	bIsLooping = false;
-	bRandomStartTime = false;
-
-	bOverrideSpawnCondition = false;
-}
-
-bool FGameplayCueNotify_CameraAnimInfo::PlayCameraAnim(const FGameplayCueNotify_SpawnContext& SpawnContext, FGameplayCueNotify_SpawnResult& OutSpawnResult) const
-{
-	if (!CameraAnim || (AnimScale <= 0.0f) || (PlayRate <= 0.0f))
-	{
-		return false;
-	}
-
-	const FGameplayCueNotify_SpawnCondition& SpawnCondition = SpawnContext.GetSpawnCondition(bOverrideSpawnCondition, SpawnConditionOverride);
-
-	if (!SpawnCondition.ShouldSpawn(SpawnContext))
-	{
-		return false;
-	}
-
-	// Camera animations only play on the target actor when it is locally controlled.
-	APlayerController* TargetPC = SpawnContext.FindLocalPlayerController(EGameplayCueNotify_LocallyControlledSource::TargetActor);
-	if (TargetPC && TargetPC->PlayerCameraManager)
-	{
-		const ECameraShakePlaySpace CameraAnimPlaySpace = GetCameraShakePlaySpace(PlaySpace);
-
-		OutSpawnResult.CameraAnim = TargetPC->PlayerCameraManager->PlayCameraAnim(CameraAnim, PlayRate, AnimScale, BlendInTime, BlendOutTime, bIsLooping, bRandomStartTime, 0.0f, CameraAnimPlaySpace, FRotator::ZeroRotator);
-	}
-
-	return (OutSpawnResult.CameraAnim != nullptr);
-}
-
-void FGameplayCueNotify_CameraAnimInfo::ValidateBurstAssets(UObject* ContainingAsset, const FString& Context, TArray<FText>& ValidationErrors) const
-{
-#if WITH_EDITORONLY_DATA
-	if (CameraAnim)
-	{
-		if (bIsLooping)
-		{
-			ValidationErrors.Add(FText::Format(
-				LOCTEXT("CameraAnim_ShouldNotLoop", "Camera animation [{0}] used in slot [{1}] for asset [{2}] is set to looping, but the slot is a one-shot (the instance will leak)."),
-				FText::AsCultureInvariant(CameraAnim->GetPathName()),
-				FText::AsCultureInvariant(Context),
-				FText::AsCultureInvariant(ContainingAsset->GetPathName())));
-		}
-	}
-#endif // #if WITH_EDITORONLY_DATA
-}
-
-
 //////////////////////////////////////////////////////////////////////////
 // FGameplayCueNotify_CameraLensEffectInfo
 //////////////////////////////////////////////////////////////////////////
@@ -1028,7 +962,6 @@ void FGameplayCueNotify_BurstEffects::ExecuteEffects(const FGameplayCueNotify_Sp
 	}
 
 	BurstCameraShake.PlayCameraShake(SpawnContext, OutSpawnResult);
-	BurstCameraAnim.PlayCameraAnim(SpawnContext, OutSpawnResult);
 	BurstCameraLensEffect.PlayCameraLensEffect(SpawnContext, OutSpawnResult);
 	BurstForceFeedback.PlayForceFeedback(SpawnContext, OutSpawnResult);
 	BurstDecal.SpawnDecal(SpawnContext, OutSpawnResult);
@@ -1047,7 +980,6 @@ void FGameplayCueNotify_BurstEffects::ValidateAssociatedAssets(UObject* Containi
 	}
 
 	BurstCameraShake.ValidateBurstAssets(ContainingAsset, Context + TEXT(".BurstCameraShake"), ValidationErrors);
-	BurstCameraAnim.ValidateBurstAssets(ContainingAsset, Context + TEXT(".BurstCameraAnim"), ValidationErrors);
 	BurstCameraLensEffect.ValidateBurstAssets(ContainingAsset, Context + TEXT(".BurstCameraLensEffect"), ValidationErrors);
 	BurstForceFeedback.ValidateBurstAssets(ContainingAsset, Context + TEXT(".BurstForceFeedback"), ValidationErrors);
 }
@@ -1079,7 +1011,6 @@ void FGameplayCueNotify_LoopingEffects::StartEffects(const FGameplayCueNotify_Sp
 	}
 
 	LoopingCameraShake.PlayCameraShake(SpawnContext, OutSpawnResult);
-	LoopingCameraAnim.PlayCameraAnim(SpawnContext, OutSpawnResult);
 	LoopingCameraLensEffect.PlayCameraLensEffect(SpawnContext, OutSpawnResult);
 	LoopingForceFeedback.PlayForceFeedback(SpawnContext, OutSpawnResult);
 }
@@ -1134,13 +1065,6 @@ void FGameplayCueNotify_LoopingEffects::StopEffects(FGameplayCueNotify_SpawnResu
 		}
 	}
 
-	// Stop the camera animation.
-	if (SpawnResult.CameraAnim)
-	{
-		const bool bStopImmediately = false;
-		SpawnResult.CameraAnim->Stop(bStopImmediately);
-	}
-
 	// Stop the camera lens effect.
 	for (TScriptInterface<ICameraLensEffectInterface> CameraLensEffect : SpawnResult.CameraLensEffects)
 	{
diff --git a/Engine/Plugins/Runtime/GameplayAbilities/Source/GameplayAbilities/Public/GameplayCueNotifyTypes.h b/Engine/Plugins/Runtime/GameplayAbilities/Source/GameplayAbilities/Public/GameplayCueNotifyTypes.h
--- a/Engine/Plugins/Runtime/GameplayAbilities/Source/GameplayAbilities/Public/GameplayCueNotifyTypes.h
+++ b/Engine/Plugins/Runtime/GameplayAbilities/Source/GameplayAbilities/Public/GameplayCueNotifyTypes.h
@@ -15,8 +15,6 @@ class UFXSystemComponent;
 class UAudioComponent;
 class USoundBase;
 class UCameraShakeBase;
-class UCameraAnim;
-class UCameraAnimInst;
 class ICameraLensEffectInterface;
 class UForceFeedbackEffect;
 class UForceFeedbackAttenuation;
@@ -264,7 +262,6 @@ public:
 		FxSystemComponents.Reset();
 		AudioComponents.Reset();
 		CameraShakes.Reset();
-		CameraAnim = nullptr;
 		CameraLensEffects.Reset();
 		ForceFeedbackComponent = nullptr;
 		ForceFeedbackTargetPC = nullptr;
@@ -285,10 +282,6 @@ public:
 	UPROPERTY(BlueprintReadOnly, Transient, Category = GameplayCueNotify)
 	TArray<UCameraShakeBase*> CameraShakes;
 
-	// Camera animation played.
-	UPROPERTY(BlueprintReadOnly, Transient, Category = GameplayCueNotify)
-	UCameraAnimInst* CameraAnim;
-
 	// List of camera len effects spawned.  There will be one camera lens effect per local player controller if the effect is played in world.
 	UPROPERTY(BlueprintReadOnly, Transient, Category = GameplayCueNotify)
 	TArray<TScriptInterface<ICameraLensEffectInterface>> CameraLensEffects;
@@ -501,68 +494,6 @@ public:
 };
 
 
-/**
- * FGameplayCueNotify_CameraAnimInfo
- *
- *	Properties that specify how to play a camera animation.
- */
-USTRUCT(BlueprintType)
-struct FGameplayCueNotify_CameraAnimInfo
-{
-	GENERATED_BODY()
-
-public:
-
-	FGameplayCueNotify_CameraAnimInfo();
-
-	bool PlayCameraAnim(const FGameplayCueNotify_SpawnContext& SpawnContext, FGameplayCueNotify_SpawnResult& OutSpawnResult) const;
-
-	void ValidateBurstAssets(UObject* ContainingAsset, const FString& Context, TArray<FText>& ValidationErrors) const;
-
-public:
-
-	// Condition to check before playing the camera animation.
-	UPROPERTY(EditDefaultsOnly, BlueprintReadOnly, Category = GameplayCueNotify, Meta = (EditCondition = "bOverrideSpawnCondition"))
-	FGameplayCueNotify_SpawnCondition SpawnConditionOverride;
-
-	// Camera animation to play.
-	UPROPERTY(EditDefaultsOnly, BlueprintReadOnly, Category = GameplayCueNotify)
-	UCameraAnim* CameraAnim;
-
-	// Scale applied to the camera animation.
-	UPROPERTY(EditDefaultsOnly, BlueprintReadOnly, Category = GameplayCueNotify)
-	float AnimScale;
-
-	// Rate to play the camera animation at.
-	UPROPERTY(EditDefaultsOnly, BlueprintReadOnly, Category = GameplayCueNotify)
-	float PlayRate;
-
-	// Time (in seconds) used to blend in the camera animation.
-	UPROPERTY(EditDefaultsOnly, BlueprintReadOnly, Category = GameplayCueNotify)
-	float BlendInTime;
-
-	// Time (in seconds) used to blend out the camera animation.
-	UPROPERTY(EditDefaultsOnly, BlueprintReadOnly, Category = GameplayCueNotify)
-	float BlendOutTime;
-
-	// What coordinate space to play the camera animation relative to.
-	UPROPERTY(EditDefaultsOnly, BlueprintReadOnly, Category = GameplayCueNotify)
-	EGameplayCueNotify_EffectPlaySpace PlaySpace;
-
-	// If enabled, the camera animation will be set to loop.
-	UPROPERTY(EditDefaultsOnly, BlueprintReadOnly, Category = GameplayCueNotify)
-	uint32 bIsLooping : 1;
-
-	// If enabled, the camera animation will start at a random time.  Only used when looping.
-	UPROPERTY(EditDefaultsOnly, BlueprintReadOnly, Category = GameplayCueNotify, Meta = (EditCondition = "bIsLooping"))
-	uint32 bRandomStartTime : 1;
-
-	// If enabled, use the spawn condition override and not the default one.
-	UPROPERTY(EditDefaultsOnly, BlueprintReadOnly, Category = GameplayCueNotify, Meta = (InlineEditConditionToggle))
-	uint32 bOverrideSpawnCondition : 1;
-};
-
-
 /**
  * FGameplayCueNotify_CameraLensEffectInfo
  *
@@ -768,10 +699,6 @@ protected:
 	UPROPERTY(EditDefaultsOnly, BlueprintReadOnly, Category = GameplayCueNotify)
 	FGameplayCueNotify_CameraShakeInfo BurstCameraShake;
 
-	// Camera animation to be played on gameplay cue execution.  This should never use a looping effect!
-	UPROPERTY(EditDefaultsOnly, BlueprintReadOnly, Category = GameplayCueNotify)
-	FGameplayCueNotify_CameraAnimInfo BurstCameraAnim;
-
 	// Camera lens effect to be played on gameplay cue execution.  This should never use a looping effect!
 	UPROPERTY(EditDefaultsOnly, BlueprintReadOnly, Category = GameplayCueNotify)
 	FGameplayCueNotify_CameraLensEffectInfo BurstCameraLensEffect;
@@ -819,10 +746,6 @@ protected:
 	UPROPERTY(EditDefaultsOnly, BlueprintReadOnly, Category = GameplayCueNotify)
 	FGameplayCueNotify_CameraShakeInfo LoopingCameraShake;
 
-	// Camera animation to be played on gameplay cue loop start.
-	UPROPERTY(EditDefaultsOnly, BlueprintReadOnly, Category = GameplayCueNotify)
-	FGameplayCueNotify_CameraAnimInfo LoopingCameraAnim;
-
 	// Camera lens effect to be played on gameplay cue loop start.
 	UPROPERTY(EditDefaultsOnly, BlueprintReadOnly, Category = GameplayCueNotify)
 	FGameplayCueNotify_CameraLensEffectInfo LoopingCameraLensEffect;
diff --git a/Engine/Plugins/Tests/EditorTests/Source/EditorTests/Private/UnrealEd/EditorAssetAutomationTests.cpp b/Engine/Plugins/Tests/EditorTests/Source/EditorTests/Private/UnrealEd/EditorAssetAutomationTests.cpp
--- a/Engine/Plugins/Tests/EditorTests/Source/EditorTests/Private/UnrealEd/EditorAssetAutomationTests.cpp
+++ b/Engine/Plugins/Tests/EditorTests/Source/EditorTests/Private/UnrealEd/EditorAssetAutomationTests.cpp
@@ -32,7 +32,6 @@
 #include "Engine/Font.h"
 #include "Engine/DataAsset.h"
 #include "Animation/AnimBlueprint.h"
-#include "Camera/CameraAnim.h"
 #include "Sound/DialogueVoice.h"
 #include "Sound/DialogueWave.h"
 #include "Exporters/Exporter.h"
@@ -97,7 +96,6 @@
 #include "ContentBrowserModule.h"
 
 #include "Slate/SlateBrushAsset.h"
-#include "Matinee/InterpData.h"
 #include "Framework/Styling/ButtonWidgetStyle.h"
 
 #include "Tests/AutomationCommon.h"
diff --git a/Engine/Plugins/VirtualProduction/Takes/Source/TakeRecorder/Private/Recorder/TakeRecorder.cpp b/Engine/Plugins/VirtualProduction/Takes/Source/TakeRecorder/Private/Recorder/TakeRecorder.cpp
--- a/Engine/Plugins/VirtualProduction/Takes/Source/TakeRecorder/Private/Recorder/TakeRecorder.cpp
+++ b/Engine/Plugins/VirtualProduction/Takes/Source/TakeRecorder/Private/Recorder/TakeRecorder.cpp
@@ -930,7 +930,7 @@ void UTakeRecorder::PreRecord()
 	check(RecordingWorld);
 	if (AWorldSettings* WorldSettings = RecordingWorld->GetWorldSettings())
 	{
-		const float ExistingMatineeTimeDilation = WorldSettings->MatineeTimeDilation;
+		const float ExistingCinematicTimeDilation = WorldSettings->CinematicTimeDilation;
 
 		const bool bInvalidTimeDilation = Parameters.User.EngineTimeDilation == 0.f;
 
@@ -939,16 +939,16 @@ void UTakeRecorder::PreRecord()
 			UE_LOG(LogTakesCore, Warning, TEXT("Time dilation cannot be 0. Ignoring time dilation for this recording."));
 		}
 
-		if (Parameters.User.EngineTimeDilation != ExistingMatineeTimeDilation && !bInvalidTimeDilation)
+		if (Parameters.User.EngineTimeDilation != ExistingCinematicTimeDilation && !bInvalidTimeDilation)
 		{
-			WorldSettings->MatineeTimeDilation = Parameters.User.EngineTimeDilation;
+			WorldSettings->CinematicTimeDilation = Parameters.User.EngineTimeDilation;
 
 			// Restore it when we're done
-			auto RestoreTimeDilation = [ExistingMatineeTimeDilation, WeakWorldSettings = MakeWeakObjectPtr(WorldSettings)]
+			auto RestoreTimeDilation = [ExistingCinematicTimeDilation, WeakWorldSettings = MakeWeakObjectPtr(WorldSettings)]
 			{
 				if (AWorldSettings* CleaupWorldSettings = WeakWorldSettings.Get())
 				{
-					CleaupWorldSettings->MatineeTimeDilation = ExistingMatineeTimeDilation;
+					CleaupWorldSettings->CinematicTimeDilation = ExistingCinematicTimeDilation;
 				}
 			};
 			OnStopCleanup.Add(RestoreTimeDilation);
diff --git a/Engine/Source/Developer/AssetTools/Private/AssetTools.cpp b/Engine/Source/Developer/AssetTools/Private/AssetTools.cpp
--- a/Engine/Source/Developer/AssetTools/Private/AssetTools.cpp
+++ b/Engine/Source/Developer/AssetTools/Private/AssetTools.cpp
@@ -53,7 +53,6 @@
 #include "AssetTypeActions/AssetTypeActions_AimOffset.h"
 #include "AssetTypeActions/AssetTypeActions_BlendSpace1D.h"
 #include "AssetTypeActions/AssetTypeActions_AimOffset1D.h"
-#include "AssetTypeActions/AssetTypeActions_CameraAnim.h"
 #include "AssetTypeActions/AssetTypeActions_TextureRenderTarget2D.h"
 #include "AssetTypeActions/AssetTypeActions_CanvasRenderTarget2D.h"
 #include "AssetTypeActions/AssetTypeActions_CurveFloat.h"
@@ -78,7 +77,6 @@
 #include "AssetTypeActions/AssetTypeActions_SubsurfaceProfile.h"
 #include "AssetTypeActions/AssetTypeActions_ActorFoliageSettings.h"
 #include "AssetTypeActions/AssetTypeActions_InstancedFoliageSettings.h"
-#include "AssetTypeActions/AssetTypeActions_InterpData.h"
 #include "AssetTypeActions/AssetTypeActions_LandscapeLayer.h"
 #include "AssetTypeActions/AssetTypeActions_LandscapeGrassType.h"
 #include "AssetTypeActions/AssetTypeActions_LightWeightInstance.h"
@@ -261,7 +259,6 @@ UAssetToolsImpl::UAssetToolsImpl(const FObjectInitializer& ObjectInitializer)
 	RegisterAssetTypeActions(MakeShareable(new FAssetTypeActions_BlendSpace1D));
 	RegisterAssetTypeActions(MakeShareable(new FAssetTypeActions_Blueprint));
 	RegisterAssetTypeActions(MakeShareable(new FAssetTypeActions_BlueprintGeneratedClass));
-	RegisterAssetTypeActions(MakeShareable(new FAssetTypeActions_CameraAnim));
 	RegisterAssetTypeActions(MakeShareable(new FAssetTypeActions_CanvasRenderTarget2D));
 	RegisterAssetTypeActions(MakeShareable(new FAssetTypeActions_Curve));
 	RegisterAssetTypeActions(MakeShareable(new FAssetTypeActions_CurveFloat));
@@ -286,7 +283,6 @@ UAssetToolsImpl::UAssetToolsImpl(const FObjectInitializer& ObjectInitializer)
 	RegisterAssetTypeActions(MakeShareable(new FAssetTypeActions_SubsurfaceProfile));
 	RegisterAssetTypeActions(MakeShareable(new FAssetTypeActions_InstancedFoliageSettings(FoliageCategoryBit)));
 	RegisterAssetTypeActions(MakeShareable(new FAssetTypeActions_ActorFoliageSettings(FoliageCategoryBit)));
-	RegisterAssetTypeActions(MakeShareable(new FAssetTypeActions_InterpData));
 	RegisterAssetTypeActions(MakeShareable(new FAssetTypeActions_LandscapeLayer));
 	RegisterAssetTypeActions(MakeShareable(new FAssetTypeActions_LandscapeGrassType(FoliageCategoryBit)));
 	RegisterAssetTypeActions(MakeShareable(new FAssetTypeActions_LightWeightInstance));
diff --git a/Engine/Source/Editor/BlueprintGraph/Classes/EdGraphSchema_K2.h b/Engine/Source/Editor/BlueprintGraph/Classes/EdGraphSchema_K2.h
--- a/Engine/Source/Editor/BlueprintGraph/Classes/EdGraphSchema_K2.h
+++ b/Engine/Source/Editor/BlueprintGraph/Classes/EdGraphSchema_K2.h
@@ -375,7 +375,6 @@ class BLUEPRINTGRAPH_API UEdGraphSchema_K2 : public UEdGraphSchema
 	static const FName PN_CastSucceeded;    // Category=PC_Exec, singleton, output
 	static const FName PN_CastFailed;    // Category=PC_Exec, singleton, output
 	static const FString PN_CastedValuePrefix;    // Category=PC_Object, singleton, output; actual pin name varies depending on the type to be casted to, this is just a prefix
-	static const FName PN_MatineeFinished;    // Category=PC_Exec, singleton, output
 
 	// construction script function names
 	static const FName FN_UserConstructionScript;
diff --git a/Engine/Source/Editor/BlueprintGraph/Private/EdGraphSchema_K2.cpp b/Engine/Source/Editor/BlueprintGraph/Private/EdGraphSchema_K2.cpp
--- a/Engine/Source/Editor/BlueprintGraph/Private/EdGraphSchema_K2.cpp
+++ b/Engine/Source/Editor/BlueprintGraph/Private/EdGraphSchema_K2.cpp
@@ -673,7 +673,6 @@ const FName UEdGraphSchema_K2::PN_Item(TEXT("Item"));
 const FName UEdGraphSchema_K2::PN_CastSucceeded(TEXT("then"));
 const FName UEdGraphSchema_K2::PN_CastFailed(TEXT("CastFailed"));
 const FString UEdGraphSchema_K2::PN_CastedValuePrefix(TEXT("As"));
-const FName UEdGraphSchema_K2::PN_MatineeFinished(TEXT("Finished"));
 
 const FName UEdGraphSchema_K2::FN_UserConstructionScript(TEXT("UserConstructionScript"));
 const FName UEdGraphSchema_K2::FN_ExecuteUbergraphBase(TEXT("ExecuteUbergraph"));
diff --git a/Engine/Source/Editor/BlueprintGraph/Public/BlueprintGraphDefinitions.h b/Engine/Source/Editor/BlueprintGraph/Public/BlueprintGraphDefinitions.h
--- a/Engine/Source/Editor/BlueprintGraph/Public/BlueprintGraphDefinitions.h
+++ b/Engine/Source/Editor/BlueprintGraph/Public/BlueprintGraphDefinitions.h
@@ -54,7 +54,6 @@
 #include "K2Node_MacroInstance.h"
 #include "K2Node_MakeArray.h"
 #include "K2Node_MathExpression.h"
-#include "K2Node_MatineeController.h"
 #include "K2Node_RemoveDelegate.h"
 #include "K2Node_Select.h"
 #include "K2Node_Self.h"
diff --git a/Engine/Source/Editor/DetailCustomizations/Private/ActorDetails.cpp b/Engine/Source/Editor/DetailCustomizations/Private/ActorDetails.cpp
--- a/Engine/Source/Editor/DetailCustomizations/Private/ActorDetails.cpp
+++ b/Engine/Source/Editor/DetailCustomizations/Private/ActorDetails.cpp
@@ -361,9 +361,7 @@ TSharedRef<SWidget> FActorDetails::OnGetConvertContent()
 
 EVisibility FActorDetails::GetConvertMenuVisibility() const
 {
-	return GLevelEditorModeTools().EnsureNotInMode(FBuiltinEditorModes::EM_InterpEdit) ?
-		EVisibility::Visible :
-		EVisibility::Collapsed;
+	return EVisibility::Visible;
 }
 
 TSharedRef<SWidget> FActorDetails::MakeConvertMenu( const FSelectedActorInfo& SelectedActorInfo )
diff --git a/Engine/Source/Editor/EditorFramework/Private/EditorModes.cpp b/Engine/Source/Editor/EditorFramework/Private/EditorModes.cpp
--- a/Engine/Source/Editor/EditorFramework/Private/EditorModes.cpp
+++ b/Engine/Source/Editor/EditorFramework/Private/EditorModes.cpp
@@ -12,7 +12,6 @@ namespace FBuiltinEditorModes
 	const FEditorModeID EM_None = NAME_None;
 	const FEditorModeID EM_Default(TEXT("EM_Default"));
 	const FEditorModeID EM_Placement(TEXT("PLACEMENT"));
-	const FEditorModeID EM_InterpEdit(TEXT("EM_InterpEdit"));
 	const FEditorModeID EM_MeshPaint(TEXT("EM_MeshPaint"));
 	const FEditorModeID EM_Landscape(TEXT("EM_Landscape"));
 	const FEditorModeID EM_Foliage(TEXT("EM_Foliage"));
diff --git a/Engine/Source/Editor/EditorFramework/Public/EditorModes.h b/Engine/Source/Editor/EditorFramework/Public/EditorModes.h
--- a/Engine/Source/Editor/EditorFramework/Public/EditorModes.h
+++ b/Engine/Source/Editor/EditorFramework/Public/EditorModes.h
@@ -21,9 +21,6 @@ namespace FBuiltinEditorModes
 	/** Placement mode */
 	EDITORFRAMEWORK_API extern const FEditorModeID EM_Placement;
 
-	/** Interpolation editing. */
-	EDITORFRAMEWORK_API extern const FEditorModeID EM_InterpEdit;
-
 	/** Mesh paint tool */
 	EDITORFRAMEWORK_API extern const FEditorModeID EM_MeshPaint;
 
diff --git a/Engine/Source/Editor/Kismet/Private/BlueprintDetailsCustomization.cpp b/Engine/Source/Editor/Kismet/Private/BlueprintDetailsCustomization.cpp
--- a/Engine/Source/Editor/Kismet/Private/BlueprintDetailsCustomization.cpp
+++ b/Engine/Source/Editor/Kismet/Private/BlueprintDetailsCustomization.cpp
@@ -502,7 +502,7 @@ void FBlueprintVarActionDetails::CustomizeDetails( IDetailLayoutBuilder& DetailL
 		.ToolTip(PrivateTooltip)
 	];
 
-	TSharedPtr<SToolTip> ExposeToCinematicsTooltip = IDocumentation::Get()->CreateToolTip(LOCTEXT("VariableExposeToCinematics_Tooltip", "Should this variable be exposed for Matinee or Sequencer to modify?"), NULL, DocLink, TEXT("ExposeToCinematics"));
+	TSharedPtr<SToolTip> ExposeToCinematicsTooltip = IDocumentation::Get()->CreateToolTip(LOCTEXT("VariableExposeToCinematics_Tooltip", "Should this variable be exposed for Sequencer to modify?"), NULL, DocLink, TEXT("ExposeToCinematics"));
 
 	Category.AddCustomRow( LOCTEXT("VariableExposeToCinematics", "Expose to Cinematics") )
 	.Visibility(TAttribute<EVisibility>(this, &FBlueprintVarActionDetails::ExposeToCinematicsVisibility))
diff --git a/Engine/Source/Editor/LevelEditor/Private/LevelEditorToolBar.cpp b/Engine/Source/Editor/LevelEditor/Private/LevelEditorToolBar.cpp
--- a/Engine/Source/Editor/LevelEditor/Private/LevelEditorToolBar.cpp
+++ b/Engine/Source/Editor/LevelEditor/Private/LevelEditorToolBar.cpp
@@ -2084,7 +2084,7 @@ void FLevelEditorToolBar::RegisterCinematicsMenu()
 
 void FLevelEditorToolBar::OnCinematicsActorPicked( AActor* Actor )
 {
-	//The matinee editor will not tick unless the editor viewport is in realtime mode.
+	//The sequencer editor will not tick unless the editor viewport is in realtime mode.
 	//the scene outliner eats input, so we must close any popups manually.
 	FSlateApplication::Get().DismissAllMenus();
 
diff --git a/Engine/Source/Editor/LevelEditor/Public/LevelEditorActions.h b/Engine/Source/Editor/LevelEditor/Public/LevelEditorActions.h
--- a/Engine/Source/Editor/LevelEditor/Public/LevelEditorActions.h
+++ b/Engine/Source/Editor/LevelEditor/Public/LevelEditorActions.h
@@ -397,10 +397,6 @@ public:
 	/** Selects all actors using the same material(s) as the current selection */
 	TSharedPtr< FUICommandInfo > SelectAllWithSameMaterial;
 
-	/** Selects all actors used by currently selected matinee actor */
-	UE_DEPRECATED(5.0, "Matinee is no longer part of the editor.")
-	TSharedPtr< FUICommandInfo > SelectAllActorsControlledByMatinee;
-
 	/** Selects all emitters using the same particle system as the current selection */
 	TSharedPtr< FUICommandInfo > SelectMatchingEmitter;
 
diff --git a/Engine/Source/Editor/LevelEditor/Public/SLevelViewport.h b/Engine/Source/Editor/LevelEditor/Public/SLevelViewport.h
--- a/Engine/Source/Editor/LevelEditor/Public/SLevelViewport.h
+++ b/Engine/Source/Editor/LevelEditor/Public/SLevelViewport.h
@@ -565,7 +565,7 @@ private:
 	void OnUseDefaultShowFlags(bool bUseSavedDefaults = false);
 
 	/**
-	 * Called to toggle allowing matinee to use this viewport to preview in
+	 * Called to toggle allowing sequencer to use this viewport to preview in
 	 */
 	void OnToggleAllowCinematicPreview();
 
diff --git a/Engine/Source/Editor/MovieSceneTools/Private/MovieSceneToolHelpers.cpp b/Engine/Source/Editor/MovieSceneTools/Private/MovieSceneToolHelpers.cpp
--- a/Engine/Source/Editor/MovieSceneTools/Private/MovieSceneToolHelpers.cpp
+++ b/Engine/Source/Editor/MovieSceneTools/Private/MovieSceneToolHelpers.cpp
@@ -38,7 +38,6 @@
 #include "MessageLogModule.h"
 #include "IMessageLogListing.h"
 #include "FbxImporter.h"
-#include "MatineeImportTools.h"
 #include "MovieSceneToolsProjectSettings.h"
 #include "MovieSceneToolsUserSettings.h"
 #include "CineCameraActor.h"
@@ -51,12 +50,6 @@
 #include "Editor/EditorEngine.h"
 #include "LevelEditorViewport.h"
 #include "AssetToolsModule.h"
-#include "Camera/CameraAnim.h"
-#include "Matinee/InterpGroup.h"
-#include "Matinee/InterpGroupInst.h"
-#include "Matinee/InterpTrackMove.h"
-#include "Matinee/InterpTrackMoveAxis.h"
-#include "Matinee/InterpTrackInstMove.h"
 #include "Channels/MovieSceneChannelProxy.h"
 #include "Evaluation/MovieSceneEvaluationTrack.h"
 #include "Evaluation/MovieSceneEvaluationTemplateInstance.h"
@@ -1047,10 +1040,9 @@ bool ImportFBXPropertyToPropertyTrack(UMovieScene* MovieScene, const FGuid Prope
 		}
 
 		FFrameNumber KeyTime = (Key.Time * FrameRate).RoundToFrame();
-		FMatineeImportTools::SetOrAddKey(ChannelData, KeyTime, Key.Value, ArriveTangent, LeaveTangent,
-				MovieSceneToolHelpers::RichCurveInterpolationToMatineeInterpolation(Key.InterpMode, Key.TangentMode), FrameRate, Key.TangentWeightMode,
+		MovieSceneToolHelpers::SetOrAddKey(ChannelData, KeyTime, Key.Value, ArriveTangent, LeaveTangent,
+				Key.InterpMode, Key.TangentMode, FrameRate, Key.TangentWeightMode,
 				Key.ArriveTangentWeight, Key.LeaveTangentWeight);
-
 	}
 
 	Channel->AutoSetTangents();
@@ -1234,10 +1226,9 @@ void ImportTransformChannelToBezierChannel(const FRichCurve& Source, ChannelType
 
 		FFrameNumber KeyTime = (Key.Time * DestFrameRate).RoundToFrame();
 		float Value = !bNegateValue ? Key.Value : -Key.Value;
-		FMatineeImportTools::SetOrAddKey(ChannelData, KeyTime + StartFrame, Value, ArriveTangent, LeaveTangent,
-			MovieSceneToolHelpers::RichCurveInterpolationToMatineeInterpolation(Key.InterpMode, Key.TangentMode), DestFrameRate, Key.TangentWeightMode,
+		MovieSceneToolHelpers::SetOrAddKey(ChannelData, KeyTime + StartFrame, Value, ArriveTangent, LeaveTangent,
+			Key.InterpMode, Key.TangentMode, DestFrameRate, Key.TangentWeightMode,
 			Key.ArriveTangentWeight, Key.LeaveTangentWeight);
-
 	}
 
 	Dest->AutoSetTangents();
@@ -3634,54 +3625,6 @@ bool MovieSceneToolHelpers::ImportFBXWithDialog(UMovieSceneSequence* InSequence,
 	return true;
 }
 
-
-EInterpCurveMode MovieSceneToolHelpers::RichCurveInterpolationToMatineeInterpolation( ERichCurveInterpMode InterpMode, ERichCurveTangentMode TangentMode)
-{
-	switch ( InterpMode )
-	{
-	case ERichCurveInterpMode::RCIM_Constant:
-		return CIM_Constant;
-	case ERichCurveInterpMode::RCIM_Cubic:
-		if (TangentMode == RCTM_Auto)
-		{
-			return CIM_CurveAuto;
-		}
-		else if (TangentMode == RCTM_Break)
-		{
-			return CIM_CurveBreak;
-		}
-		return CIM_CurveUser;  
-	case ERichCurveInterpMode::RCIM_Linear:
-		return CIM_Linear;
-	default:
-		return CIM_CurveAuto;
-	}
-}
-
-void MovieSceneToolHelpers::CopyKeyDataToMoveAxis(const TMovieSceneChannelData<FMovieSceneDoubleValue>& Channel, UInterpTrackMoveAxis* MoveAxis, FFrameRate InFrameRate)
-{
-	MoveAxis->FloatTrack.Points.Reset();
-
-	static FName LookupName(NAME_None);
-	
-	TArrayView<const FFrameNumber>           Times  = Channel.GetTimes();
-	TArrayView<const FMovieSceneDoubleValue> Values = Channel.GetValues();
-
-	for (int32 KeyIndex = 0; KeyIndex < Times.Num(); ++KeyIndex)
-	{
-		const float Time = Times[KeyIndex] / InFrameRate;
-		const FMovieSceneDoubleValue& Value = Values[KeyIndex];
-
-		const int32 PointIndex = MoveAxis->FloatTrack.AddPoint(Time, Value.Value);
-		MoveAxis->LookupTrack.AddPoint(Time, LookupName);
-
-		FInterpCurvePoint<float>& Point = MoveAxis->FloatTrack.Points[PointIndex];
-		Point.ArriveTangent = Value.Tangent.ArriveTangent * InFrameRate.AsDecimal();
-		Point.LeaveTangent = Value.Tangent.LeaveTangent * InFrameRate.AsDecimal();
-		Point.InterpMode = RichCurveInterpolationToMatineeInterpolation(Value.InterpMode, Value.TangentMode);
-	}
-}
-
 bool MovieSceneToolHelpers::HasHiddenMobility(const UClass* ObjectClass)
 {
 	if (ObjectClass)
@@ -3747,7 +3690,7 @@ bool MovieSceneToolHelpers::ExportFBX(UWorld* World, UMovieScene* MovieScene, IM
 	UnFbx::FFbxExporter* Exporter = UnFbx::FFbxExporter::GetInstance();
 
 	Exporter->CreateDocument();
-	Exporter->SetTrasformBaking(false);
+	Exporter->SetTransformBaking(false);
 	Exporter->SetKeepHierarchy(true);
 
 	ExportLevelMesh(Exporter, World->PersistentLevel, Player, Bindings, NodeNameAdapter, Template);
@@ -3820,7 +3763,7 @@ static void TickFrame(const FFrameNumber& FrameNumber,float DeltaTime, UMovieSce
 			BakeHelper->PreEvaluation(MovieScene, FrameNumber);
 		}
 	}
-	// This will call UpdateSkelPose on the skeletal mesh component to move bones based on animations in the matinee group
+	// This will call UpdateSkelPose on the skeletal mesh component to move bones based on animations in the sequence
 	int32 Index = FrameNumber.Value;
 	AnimTrackAdapter.UpdateAnimation(Index);
 	for (IMovieSceneToolsAnimationBakeHelper* BakeHelper : BakeHelpers)
@@ -4499,6 +4442,44 @@ void MovieSceneToolHelpers::SetOrAddKey(TMovieSceneChannelData<FMovieSceneFloatV
 	}
 }
 
+void MovieSceneToolHelpers::SetOrAddKey(TMovieSceneChannelData<FMovieSceneFloatValue>& ChannelData, FFrameNumber Time, float Value, 
+		float ArriveTangent, float LeaveTangent, ERichCurveInterpMode InterpMode, ERichCurveTangentMode TangentMode, 
+		FFrameRate FrameRate, ERichCurveTangentWeightMode WeightedMode, float ArriveTangentWeight, float LeaveTangentWeight)
+{
+	if (ChannelData.FindKey(Time) == INDEX_NONE)
+	{
+		FMovieSceneFloatValue NewKey(Value);
+
+		NewKey.InterpMode = InterpMode;
+		NewKey.TangentMode = TangentMode;
+		NewKey.Tangent.ArriveTangent = ArriveTangent / FrameRate.AsDecimal();
+		NewKey.Tangent.LeaveTangent = LeaveTangent / FrameRate.AsDecimal();
+		NewKey.Tangent.TangentWeightMode = WeightedMode;
+		NewKey.Tangent.ArriveTangentWeight = ArriveTangentWeight;
+		NewKey.Tangent.LeaveTangentWeight = LeaveTangentWeight;
+		ChannelData.AddKey(Time, NewKey);
+	}
+}
+
+void MovieSceneToolHelpers::SetOrAddKey(TMovieSceneChannelData<FMovieSceneDoubleValue>& ChannelData, FFrameNumber Time, double Value, 
+		float ArriveTangent, float LeaveTangent, ERichCurveInterpMode InterpMode, ERichCurveTangentMode TangentMode, 
+		FFrameRate FrameRate, ERichCurveTangentWeightMode WeightedMode, float ArriveTangentWeight, float LeaveTangentWeight)
+{
+	if (ChannelData.FindKey(Time) == INDEX_NONE)
+	{
+		FMovieSceneDoubleValue NewKey(Value);
+
+		NewKey.InterpMode = InterpMode;
+		NewKey.TangentMode = TangentMode;
+		NewKey.Tangent.ArriveTangent = ArriveTangent / FrameRate.AsDecimal();
+		NewKey.Tangent.LeaveTangent = LeaveTangent / FrameRate.AsDecimal();
+		NewKey.Tangent.TangentWeightMode = WeightedMode;
+		NewKey.Tangent.ArriveTangentWeight = ArriveTangentWeight;
+		NewKey.Tangent.LeaveTangentWeight = LeaveTangentWeight;
+		ChannelData.AddKey(Time, NewKey);
+	}
+}
+
 void MovieSceneToolHelpers::SetOrAddKey(TMovieSceneChannelData<FMovieSceneDoubleValue>& ChannelData, FFrameNumber Time, double Value)
 {
 	int32 ExistingIndex = ChannelData.FindKey(Time);
diff --git a/Engine/Source/Editor/MovieSceneTools/Private/MovieSceneToolsModule.cpp b/Engine/Source/Editor/MovieSceneTools/Private/MovieSceneToolsModule.cpp
--- a/Engine/Source/Editor/MovieSceneTools/Private/MovieSceneToolsModule.cpp
+++ b/Engine/Source/Editor/MovieSceneTools/Private/MovieSceneToolsModule.cpp
@@ -47,7 +47,6 @@
 #include "TrackEditors/SpawnTrackEditor.h"
 #include "TrackEditors/LevelVisibilityTrackEditor.h"
 #include "TrackEditors/DataLayerTrackEditor.h"
-#include "TrackEditors/CameraAnimTrackEditor.h"
 #include "TrackEditors/CameraShakeTrackEditor.h"
 #include "TrackEditors/MaterialParameterCollectionTrackEditor.h"
 #include "TrackEditors/ObjectPropertyTrackEditor.h"
@@ -161,7 +160,6 @@ void FMovieSceneToolsModule::StartupModule()
 		SpawnTrackCreateEditorHandle = SequencerModule.RegisterTrackEditor( FOnCreateTrackEditor::CreateStatic( &FSpawnTrackEditor::CreateTrackEditor ) );
 		LevelVisibilityTrackCreateEditorHandle = SequencerModule.RegisterTrackEditor( FOnCreateTrackEditor::CreateStatic( &FLevelVisibilityTrackEditor::CreateTrackEditor ) );
 		DataLayerTrackCreateEditorHandle = SequencerModule.RegisterTrackEditor( FOnCreateTrackEditor::CreateStatic( &FDataLayerTrackEditor::CreateTrackEditor ) );
-		CameraAnimTrackCreateEditorHandle = SequencerModule.RegisterTrackEditor(FOnCreateTrackEditor::CreateStatic(&FCameraAnimTrackEditor::CreateTrackEditor));
 		CameraShakeTrackCreateEditorHandle = SequencerModule.RegisterTrackEditor(FOnCreateTrackEditor::CreateStatic(&FCameraShakeTrackEditor::CreateTrackEditor));
 		MPCTrackCreateEditorHandle = SequencerModule.RegisterTrackEditor(FOnCreateTrackEditor::CreateStatic(&FMaterialParameterCollectionTrackEditor::CreateTrackEditor));
 		PrimitiveMaterialCreateEditorHandle = SequencerModule.RegisterTrackEditor(FOnCreateTrackEditor::CreateStatic(&FPrimitiveMaterialTrackEditor::CreateTrackEditor));
@@ -311,7 +309,6 @@ void FMovieSceneToolsModule::ShutdownModule()
 	SequencerModule.UnRegisterTrackEditor( SpawnTrackCreateEditorHandle );
 	SequencerModule.UnRegisterTrackEditor( LevelVisibilityTrackCreateEditorHandle );
 	SequencerModule.UnRegisterTrackEditor( DataLayerTrackCreateEditorHandle );
-	SequencerModule.UnRegisterTrackEditor( CameraAnimTrackCreateEditorHandle );
 	SequencerModule.UnRegisterTrackEditor( CameraShakeTrackCreateEditorHandle );
 	SequencerModule.UnRegisterTrackEditor( MPCTrackCreateEditorHandle );
 	SequencerModule.UnRegisterTrackEditor( ObjectTrackCreateEditorHandle );
diff --git a/Engine/Source/Editor/MovieSceneTools/Private/TrackEditors/AudioTrackEditor.cpp b/Engine/Source/Editor/MovieSceneTools/Private/TrackEditors/AudioTrackEditor.cpp
--- a/Engine/Source/Editor/MovieSceneTools/Private/TrackEditors/AudioTrackEditor.cpp
+++ b/Engine/Source/Editor/MovieSceneTools/Private/TrackEditors/AudioTrackEditor.cpp
@@ -31,8 +31,6 @@
 #include "ContentBrowserModule.h"
 #include "SequencerUtilities.h"
 #include "AssetRegistry/AssetRegistryModule.h"
-#include "MatineeImportTools.h"
-#include "Matinee/InterpTrackSound.h"
 #include "ISectionLayoutBuilder.h"
 #include "MovieSceneToolHelpers.h"
 
@@ -959,14 +957,6 @@ bool FAudioTrackEditor::SupportsSequence(UMovieSceneSequence* InSequence) const
 }
 
 
-void CopyInterpSoundTrack(TSharedRef<ISequencer> Sequencer, UInterpTrackSound* MatineeSoundTrack, UMovieSceneAudioTrack* AudioTrack)
-{
-	if (FMatineeImportTools::CopyInterpSoundTrack(MatineeSoundTrack, AudioTrack))
-	{
-		Sequencer.Get().NotifyMovieSceneDataChanged( EMovieSceneDataChangeType::MovieSceneStructureItemAdded );
-	}
-}
-
 const FSlateBrush* FAudioTrackEditor::GetIconBrush() const
 {
 	return FAppStyle::GetBrush("Sequencer.Tracks.Audio");
diff --git a/Engine/Source/Editor/MovieSceneTools/Private/TrackEditors/ParticleTrackEditor.cpp b/Engine/Source/Editor/MovieSceneTools/Private/TrackEditors/ParticleTrackEditor.cpp
--- a/Engine/Source/Editor/MovieSceneTools/Private/TrackEditors/ParticleTrackEditor.cpp
+++ b/Engine/Source/Editor/MovieSceneTools/Private/TrackEditors/ParticleTrackEditor.cpp
@@ -16,8 +16,6 @@
 #include "CommonMovieSceneTools.h"
 #include "Particles/ParticleLODLevel.h"
 #include "Particles/ParticleModuleRequired.h"
-#include "Matinee/InterpTrackToggle.h"
-#include "MatineeImportTools.h"
 
 
 namespace AnimatableParticleEditorConstants
@@ -297,14 +295,4 @@ FKeyPropertyResult FParticleTrackEditor::AddKeyInternal( FFrameNumber KeyTime, U
 }
 
 
-void CopyInterpParticleTrack(TSharedRef<ISequencer> Sequencer, UInterpTrackToggle* MatineeToggleTrack, UMovieSceneParticleTrack* ParticleTrack)
-{
-	if (FMatineeImportTools::CopyInterpParticleTrack(MatineeToggleTrack, ParticleTrack))
-	{
-		Sequencer.Get().NotifyMovieSceneDataChanged( EMovieSceneDataChangeType::MovieSceneStructureItemAdded );
-	}
-}
-
-
-
 #undef LOCTEXT_NAMESPACE
diff --git a/Engine/Source/Editor/MovieSceneTools/Private/TrackEditors/PropertyTrackEditors/ColorPropertyTrackEditor.cpp b/Engine/Source/Editor/MovieSceneTools/Private/TrackEditors/PropertyTrackEditors/ColorPropertyTrackEditor.cpp
--- a/Engine/Source/Editor/MovieSceneTools/Private/TrackEditors/PropertyTrackEditors/ColorPropertyTrackEditor.cpp
+++ b/Engine/Source/Editor/MovieSceneTools/Private/TrackEditors/PropertyTrackEditors/ColorPropertyTrackEditor.cpp
@@ -4,9 +4,6 @@
 #include "Editor/UnrealEdEngine.h"
 #include "UnrealEdGlobals.h"
 #include "Sections/ColorPropertySection.h"
-#include "MatineeImportTools.h"
-#include "Matinee/InterpTrackLinearColorProp.h"
-#include "Matinee/InterpTrackColorProp.h"
 #include "Evaluation/MovieScenePropertyTemplate.h"
 #include "EntitySystem/Interrogation/MovieSceneInterrogationLinker.h"
 
@@ -79,22 +76,6 @@ void FColorPropertyTrackEditor::GenerateKeysFromPropertyChanged( const FProperty
 	OutGeneratedKeys.Add(FMovieSceneChannelValueSetter::Create<FMovieSceneFloatChannel>(3, ColorValue.A, bKeyAlpha));
 }
 
-void CopyInterpColorTrack(TSharedRef<ISequencer> Sequencer, UInterpTrackColorProp* ColorPropTrack, UMovieSceneColorTrack* ColorTrack)
-{
-	if (FMatineeImportTools::CopyInterpColorTrack(ColorPropTrack, ColorTrack))
-	{
-		Sequencer.Get().NotifyMovieSceneDataChanged( EMovieSceneDataChangeType::MovieSceneStructureItemAdded );
-	}
-}
-
-void CopyInterpLinearColorTrack(TSharedRef<ISequencer> Sequencer, UInterpTrackLinearColorProp* LinearColorPropTrack, UMovieSceneColorTrack* ColorTrack)
-{
-	if (FMatineeImportTools::CopyInterpLinearColorTrack(LinearColorPropTrack, ColorTrack))
-	{
-		Sequencer.Get().NotifyMovieSceneDataChanged( EMovieSceneDataChangeType::MovieSceneStructureItemAdded );
-	}
-}
-
 bool FColorPropertyTrackEditor::ModifyGeneratedKeysByCurrentAndWeight(UObject *Object, UMovieSceneTrack *Track, UMovieSceneSection* SectionToKey, FFrameNumber KeyTime, FGeneratedTrackKeys& GeneratedTotalKeys, float Weight) const
 {
 	using namespace UE::MovieScene;
diff --git a/Engine/Source/Editor/MovieSceneTools/Private/TrackEditors/PropertyTrackEditors/EulerTransformPropertyTrackEditor.cpp b/Engine/Source/Editor/MovieSceneTools/Private/TrackEditors/PropertyTrackEditors/EulerTransformPropertyTrackEditor.cpp
--- a/Engine/Source/Editor/MovieSceneTools/Private/TrackEditors/PropertyTrackEditors/EulerTransformPropertyTrackEditor.cpp
+++ b/Engine/Source/Editor/MovieSceneTools/Private/TrackEditors/PropertyTrackEditors/EulerTransformPropertyTrackEditor.cpp
@@ -1,7 +1,6 @@
 // Copyright Epic Games, Inc. All Rights Reserved.
 
 #include "TrackEditors/PropertyTrackEditors/EulerTransformPropertyTrackEditor.h"
-#include "MatineeImportTools.h"
 #include "UnrealEdGlobals.h"
 #include "Editor/UnrealEdEngine.h"
 #include "SequencerUtilities.h"
diff --git a/Engine/Source/Editor/MovieSceneTools/Private/TrackEditors/PropertyTrackEditors/FloatPropertyTrackEditor.cpp b/Engine/Source/Editor/MovieSceneTools/Private/TrackEditors/PropertyTrackEditors/FloatPropertyTrackEditor.cpp
--- a/Engine/Source/Editor/MovieSceneTools/Private/TrackEditors/PropertyTrackEditors/FloatPropertyTrackEditor.cpp
+++ b/Engine/Source/Editor/MovieSceneTools/Private/TrackEditors/PropertyTrackEditors/FloatPropertyTrackEditor.cpp
@@ -3,8 +3,6 @@
 #include "TrackEditors/PropertyTrackEditors/FloatPropertyTrackEditor.h"
 #include "Editor/UnrealEdEngine.h"
 #include "UnrealEdGlobals.h"
-#include "MatineeImportTools.h"
-#include "Matinee/InterpTrackFloatBase.h"
 #include "MovieSceneToolHelpers.h"
 #include "Evaluation/MovieScenePropertyTemplate.h"
 
@@ -25,15 +23,6 @@ void FFloatPropertyTrackEditor::GenerateKeysFromPropertyChanged( const FProperty
 }
 
 
-void CopyInterpFloatTrack(TSharedRef<ISequencer> Sequencer, UInterpTrackFloatBase* MatineeFloatTrack, UMovieSceneFloatTrack* FloatTrack)
-{
-	if (FMatineeImportTools::CopyInterpFloatTrack(MatineeFloatTrack, FloatTrack))
-	{
-		Sequencer.Get().NotifyMovieSceneDataChanged( EMovieSceneDataChangeType::MovieSceneStructureItemAdded );
-	}
-}
-
-
 float FFloatPropertyTrackEditor::RecomposeFloat(float InCurrentValue, UObject* AnimatedObject, UMovieSceneSection* Section)
 {
 	using namespace UE::MovieScene;
diff --git a/Engine/Source/Editor/MovieSceneTools/Private/TrackEditors/PropertyTrackEditors/VectorPropertyTrackEditor.cpp b/Engine/Source/Editor/MovieSceneTools/Private/TrackEditors/PropertyTrackEditors/VectorPropertyTrackEditor.cpp
--- a/Engine/Source/Editor/MovieSceneTools/Private/TrackEditors/PropertyTrackEditors/VectorPropertyTrackEditor.cpp
+++ b/Engine/Source/Editor/MovieSceneTools/Private/TrackEditors/PropertyTrackEditors/VectorPropertyTrackEditor.cpp
@@ -2,8 +2,6 @@
 
 #include "TrackEditors/PropertyTrackEditors/VectorPropertyTrackEditor.h"
 #include "Tracks/MovieSceneVectorTrack.h"
-#include "MatineeImportTools.h"
-#include "Matinee/InterpTrackVectorProp.h"
 #include "UnrealEdGlobals.h"
 #include "Editor/UnrealEdEngine.h"
 #include "MovieSceneToolHelpers.h"
@@ -98,14 +96,6 @@ void FFloatVectorPropertyTrackEditor::InitializeNewTrack( UMovieSceneFloatVector
 	}
 }
 
-void CopyInterpVectorTrack(TSharedRef<ISequencer> Sequencer, UInterpTrackVectorProp* MatineeVectorTrack, UMovieSceneFloatVectorTrack* VectorTrack)
-{
-	if (FMatineeImportTools::CopyInterpVectorTrack(MatineeVectorTrack, VectorTrack))
-	{
-		Sequencer.Get().NotifyMovieSceneDataChanged( EMovieSceneDataChangeType::MovieSceneStructureItemAdded );
-	}
-}
-
 bool FFloatVectorPropertyTrackEditor::ModifyGeneratedKeysByCurrentAndWeight(UObject *Object, UMovieSceneTrack *Track, UMovieSceneSection* SectionToKey, FFrameNumber KeyTime, FGeneratedTrackKeys& GeneratedTotalKeys, float Weight) const
 {
 	using namespace UE::MovieScene;
diff --git a/Engine/Source/Editor/MovieSceneTools/Private/TrackEditors/SkeletalAnimationTrackEditor.cpp b/Engine/Source/Editor/MovieSceneTools/Private/TrackEditors/SkeletalAnimationTrackEditor.cpp
--- a/Engine/Source/Editor/MovieSceneTools/Private/TrackEditors/SkeletalAnimationTrackEditor.cpp
+++ b/Engine/Source/Editor/MovieSceneTools/Private/TrackEditors/SkeletalAnimationTrackEditor.cpp
@@ -23,8 +23,6 @@
 #include "AssetRegistry/AssetRegistryModule.h"
 #include "IContentBrowserSingleton.h"
 #include "ContentBrowserModule.h"
-#include "MatineeImportTools.h"
-#include "Matinee/InterpTrackAnimControl.h"
 #include "SequencerSettings.h"
 #include "SequencerUtilities.h"
 #include "ISectionLayoutBuilder.h"
@@ -1852,16 +1850,6 @@ FKeyPropertyResult FSkeletalAnimationTrackEditor::AddKeyInternal( FFrameNumber K
 	return KeyPropertyResult;
 }
 
-void CopyInterpAnimControlTrack(TSharedRef<ISequencer> Sequencer, UInterpTrackAnimControl* MatineeAnimControlTrack, UMovieSceneSkeletalAnimationTrack* SkeletalAnimationTrack)
-{
-	FFrameNumber EndPlaybackRange = UE::MovieScene::DiscreteExclusiveUpper(Sequencer.Get().GetFocusedMovieSceneSequence()->GetMovieScene()->GetPlaybackRange());
-
-	if (FMatineeImportTools::CopyInterpAnimControlTrack(MatineeAnimControlTrack, SkeletalAnimationTrack, EndPlaybackRange))
-	{
-		Sequencer.Get().NotifyMovieSceneDataChanged( EMovieSceneDataChangeType::MovieSceneStructureItemAdded );
-	}
-}
-
 void FSkeletalAnimationTrackEditor::BuildTrackContextMenu( FMenuBuilder& MenuBuilder, UMovieSceneTrack* Track )
 {
 	//there's a bug with a section being open already, so we end it.
diff --git a/Engine/Source/Editor/MovieSceneTools/Private/TrackEditors/TransformTrackEditor.cpp b/Engine/Source/Editor/MovieSceneTools/Private/TrackEditors/TransformTrackEditor.cpp
--- a/Engine/Source/Editor/MovieSceneTools/Private/TrackEditors/TransformTrackEditor.cpp
+++ b/Engine/Source/Editor/MovieSceneTools/Private/TrackEditors/TransformTrackEditor.cpp
@@ -17,10 +17,7 @@
 #include "LevelEditorViewport.h"
 #include "UnrealEdGlobals.h"
 #include "ISectionLayoutBuilder.h"
-#include "MatineeImportTools.h"
 #include "IKeyArea.h"
-#include "Matinee/InterpTrackMove.h"
-#include "Matinee/InterpTrackMoveAxis.h"
 #include "IContentBrowserSingleton.h"
 #include "ContentBrowserModule.h"
 #include "Editor.h"
@@ -101,43 +98,6 @@ bool F3DTransformTrackEditor::SupportsType( TSubclassOf<UMovieSceneTrack> Type )
 }
 
 
-void CopyInterpMoveTrack(TSharedRef<ISequencer> Sequencer, UInterpTrackMove* MoveTrack, UMovieScene3DTransformTrack* TransformTrack)
-{
-	if (FMatineeImportTools::CopyInterpMoveTrack(MoveTrack, TransformTrack))
-	{
-		Sequencer.Get().NotifyMovieSceneDataChanged( EMovieSceneDataChangeType::MovieSceneStructureItemAdded );
-	}
-}
-
-
-bool CanCopyInterpMoveTrack(UInterpTrackMove* MoveTrack, UMovieScene3DTransformTrack* TransformTrack)
-{
-	if (!MoveTrack || !TransformTrack)
-	{
-		return false;
-	}
-
-	bool bHasKeyframes = MoveTrack->GetNumKeyframes() != 0;
-
-	for (auto SubTrack : MoveTrack->SubTracks)
-	{
-		if (SubTrack->IsA(UInterpTrackMoveAxis::StaticClass()))
-		{
-			UInterpTrackMoveAxis* MoveSubTrack = Cast<UInterpTrackMoveAxis>(SubTrack);
-			if (MoveSubTrack)
-			{
-				if (MoveSubTrack->FloatTrack.Points.Num() > 0)
-				{
-					bHasKeyframes = true;
-					break;
-				}
-			}
-		}
-	}
-		
-	return bHasKeyframes;
-}
-
 void F3DTransformTrackEditor::BuildTrackContextMenu( FMenuBuilder& MenuBuilder, UMovieSceneTrack* Track )
 {
 	UMovieScene3DTransformTrack* TransformTrack = Cast<UMovieScene3DTransformTrack>( Track );
diff --git a/Engine/Source/Editor/MovieSceneTools/Public/MovieSceneToolHelpers.h b/Engine/Source/Editor/MovieSceneTools/Public/MovieSceneToolHelpers.h
--- a/Engine/Source/Editor/MovieSceneTools/Public/MovieSceneToolHelpers.h
+++ b/Engine/Source/Editor/MovieSceneTools/Public/MovieSceneToolHelpers.h
@@ -23,7 +23,6 @@ class ISequencer;
 class UMovieScene;
 class UMovieSceneSection;
 class UMovieSceneSequence;
-class UInterpTrackMoveAxis;
 struct FMovieSceneObjectBindingID;
 class UMovieSceneTrack;
 struct FMovieSceneEvaluationTrack;
@@ -390,24 +389,6 @@ public:
 	static void CopyCameraProperties(fbxsdk::FbxCamera* CameraNode, AActor* InCameraActor);
 
 
-	/*
-	 * Rich curve interpolation to matinee interpolation
-	 *
-	 * @param InterpMode The rich curve interpolation to convert
-	 * @return The converted matinee interpolation
-	 */
-	static EInterpCurveMode RichCurveInterpolationToMatineeInterpolation( ERichCurveInterpMode InterpMode, ERichCurveTangentMode TangentMode);
-
-	/*
-	 * Copy key data to move axis
-	 *
-	 * @param KeyData The key data to copy from
-	 * @param MoveAxis The move axis to copy to
-	 * @param FrameRate The frame rate of the source channel
-	 */
-	static void CopyKeyDataToMoveAxis(const TMovieSceneChannelData<FMovieSceneDoubleValue>& KeyData, UInterpTrackMoveAxis* MoveAxis, FFrameRate FrameRate);
-
-
 	/*
 	 * Export the SkelMesh to an Anim Sequence for specified MovieScene and Player
 	 *
@@ -544,6 +525,22 @@ public:
 	* @param Value  Value to Set
 	*/
 	static void SetOrAddKey(TMovieSceneChannelData<FMovieSceneDoubleValue>& ChannelData, FFrameNumber Time, double Value);
+
+	/* 
+	* Set or add a key onto a float channel based on rich curve data.
+	*/
+	static void SetOrAddKey(TMovieSceneChannelData<FMovieSceneFloatValue>& Curve, FFrameNumber Time, float Value, 
+			float ArriveTangent, float LeaveTangent, ERichCurveInterpMode InterpMode, ERichCurveTangentMode TangentMode,
+			FFrameRate FrameRate, ERichCurveTangentWeightMode WeightedMode = RCTWM_WeightedNone, 
+			float ArriveTangentWeight = 0.0f, float LeaveTangentWeight = 0.0f);
+
+	/*
+	* Set or add a key onto a double channel based on rich curve data.
+	*/
+	static void SetOrAddKey(TMovieSceneChannelData<FMovieSceneDoubleValue>& Curve, FFrameNumber Time, double Value, 
+			float ArriveTangent, float LeaveTangent, ERichCurveInterpMode InterpMode, ERichCurveTangentMode TangentMode,
+			FFrameRate FrameRate, ERichCurveTangentWeightMode WeightedMode = RCTWM_WeightedNone, 
+			float ArriveTangentWeight = 0.0f, float LeaveTangentWeight = 0.0f);
 	
 	/*
 	*  Get an actors world transforms at the specified times using a player
diff --git a/Engine/Source/Editor/MovieSceneTools/Public/MovieSceneToolsModule.h b/Engine/Source/Editor/MovieSceneTools/Public/MovieSceneToolsModule.h
--- a/Engine/Source/Editor/MovieSceneTools/Public/MovieSceneToolsModule.h
+++ b/Engine/Source/Editor/MovieSceneTools/Public/MovieSceneToolsModule.h
@@ -117,7 +117,6 @@
 	FDelegateHandle SpawnTrackCreateEditorHandle;
 	FDelegateHandle LevelVisibilityTrackCreateEditorHandle;
 	FDelegateHandle DataLayerTrackCreateEditorHandle;
-	FDelegateHandle CameraAnimTrackCreateEditorHandle;
 	FDelegateHandle CameraShakeTrackCreateEditorHandle;
 	FDelegateHandle MPCTrackCreateEditorHandle;
 	FDelegateHandle PrimitiveMaterialCreateEditorHandle;
diff --git a/Engine/Source/Editor/Sequencer/Private/Sequencer.cpp b/Engine/Source/Editor/Sequencer/Private/Sequencer.cpp
--- a/Engine/Source/Editor/Sequencer/Private/Sequencer.cpp
+++ b/Engine/Source/Editor/Sequencer/Private/Sequencer.cpp
@@ -758,7 +758,7 @@ void FSequencer::Tick(float InDeltaTime)
 
 	UObject* PlaybackContext = GetPlaybackContext();
 	UWorld* World = PlaybackContext ? PlaybackContext->GetWorld() : nullptr;
-	float Dilation = World ? World->GetWorldSettings()->MatineeTimeDilation : 1.f;
+	float Dilation = World ? World->GetWorldSettings()->CinematicTimeDilation : 1.f;
 
 	TimeController->Tick(InDeltaTime, PlaybackSpeed * Dilation);
 
diff --git a/Engine/Source/Editor/UnrealEd/Classes/Editor/EditorEngine.h b/Engine/Source/Editor/UnrealEd/Classes/Editor/EditorEngine.h
--- a/Engine/Source/Editor/UnrealEd/Classes/Editor/EditorEngine.h
+++ b/Engine/Source/Editor/UnrealEd/Classes/Editor/EditorEngine.h
@@ -36,7 +36,6 @@
 
 class APlayerStart;
 class Error;
-class AMatineeActor;
 class FEditorViewportClient;
 class FEditorWorldManager;
 class FMessageLog;
@@ -2180,12 +2179,6 @@ public:
 	 */
 	void DeselectLevelInLevelBrowser();
 
-	/**
-	 * Selects all actors controlled by currently selected MatineeActor
-	 */
-	UE_DEPRECATED(5.0, "Matinee is no longer part of the editor.")
-	void SelectAllActorsControlledByMatinee() {}
-
 	/**
 	 * Selects all actors with the same class as the current selection
 	 *
@@ -2362,35 +2355,6 @@ public:
 	 */
 	void ConvertActorsFromClass( UClass* FromClass, UClass* ToClass );
 
-	/**
-	 * Gets a delegate that is executed when a matinee is requested to be opened
-	 *
-	 * The first parameter is the matinee actor
-	 *
-	 * @return The event delegate.
-	 */
-	UE_DEPRECATED(5.0, "Matinee is no longer part of the editor.")
-	DECLARE_DELEGATE_RetVal_OneParam(bool, FShouldOpenMatineeCallback, AMatineeActor*)
-	FShouldOpenMatineeCallback& OnShouldOpenMatinee() { return ShouldOpenMatineeCallback; }
-
-	/** 
-	 * Show a (Suppressable) warning dialog to remind the user they are about to lose their undo buffer 
-	 *
-	 * @param MatineeActor	The actor we wish to check (can be null)
-	 * @returns true if the user wishes to proceed
-	 */
-	UE_DEPRECATED(5.0, "Matinee is no longer part of the editor.")
-	bool ShouldOpenMatinee(AMatineeActor* MatineeActor) const { return false; }
-
-	/** 
-	 * Open the Matinee tool to edit the supplied MatineeActor. Will check that MatineeActor has an InterpData attached.
-	 *
-	 * @param MatineeActor	The actor we wish to edit
-	 * @param bWarnUser		If true, calls ShouldOpenMatinee as part of the open process
-	 */
-	UE_DEPRECATED(5.0, "Matinee is no longer part of the editor.")
-	void OpenMatinee(class AMatineeActor* MatineeActor, bool bWarnUser = true) {}
-
 	/**
 	* Update any outstanding reflection captures
 	*/
@@ -2897,9 +2861,6 @@ private:
 	/** Broadcasts to allow selection of unloaded actors */
 	FSelectUnloadedActorsEvent SelectUnloadedActorsEvent;
 
-	/** Delegate to be called when a matinee is requested to be opened */
-	FShouldOpenMatineeCallback ShouldOpenMatineeCallback;	
-
 	/** Reference to owner of the current popup */
 	TWeakPtr<class SWindow> PopupWindow;
 
@@ -3163,10 +3124,6 @@ private:
 
 protected:
 
-	/** Called when Matinee is opened */
-	UE_DEPRECATED(5.0, "Matinee is no longer part of the editor.")
-	virtual void OnOpenMatinee(){};
-
 	/**
 	 * Invalidates all editor viewports and hit proxies, used when global changes like Undo/Redo may have invalidated state everywhere
 	 */
diff --git a/Engine/Source/Editor/UnrealEd/Classes/Editor/Transactor.h b/Engine/Source/Editor/UnrealEd/Classes/Editor/Transactor.h
--- a/Engine/Source/Editor/UnrealEd/Classes/Editor/Transactor.h
+++ b/Engine/Source/Editor/UnrealEd/Classes/Editor/Transactor.h
@@ -253,8 +253,9 @@ protected:
 	/** The key object being edited in this transaction. For example the blueprint object. Can be NULL */
 	UObject* PrimaryObject = nullptr;
 
-	/** If true, on apply flip the direction of iteration over object records. The only client for which this is false is FMatineeTransaction */
+	/** If true, on apply flip the direction of iteration over object records. */
 	bool		bFlip = true;
+	// @TODO now that Matinee is gone, it's worth investigating removing this
 
 	/** Used to track direction to iterate over transaction's object records. Typically -1 for Undo, 1 for Redo */
 	int32		Inc = -1;
diff --git a/Engine/Source/Editor/UnrealEd/Classes/Editor/UnrealEdEngine.h b/Engine/Source/Editor/UnrealEd/Classes/Editor/UnrealEdEngine.h
--- a/Engine/Source/Editor/UnrealEd/Classes/Editor/UnrealEdEngine.h
+++ b/Engine/Source/Editor/UnrealEd/Classes/Editor/UnrealEdEngine.h
@@ -768,10 +768,6 @@ public:
 	 */
 	virtual bool CanSavePackage( UPackage* PackageToSave );
 
-	/** Converts kismet based matinees in the current level to matinees controlled via matinee actors */
-	UE_DEPRECATED(5.0, "Matinee is no longer part of the editor.")
-	void ConvertMatinees() {}
-
 	/**
 	 * Updates the volume actor visibility for all viewports based on the passed in volume class
 	 *
@@ -894,11 +890,6 @@ protected:
 	/** Instance responsible for monitoring this editor's performance */
 	FPerformanceMonitor* PerformanceMonitor;
 
-	/** Handle to the registered OnMatineeEditorClosed delegate. */
-	FDelegateHandle OnMatineeEditorClosedDelegateHandle;
-	/** Handle to the registered UpdateEdModeOnMatineeClose delegate. */
-	FDelegateHandle UpdateEdModeOnMatineeCloseDelegateHandle;
-
 	/** Whether the pivot has been moved independently */
 	bool bPivotMovedIndependently;
 
diff --git a/Engine/Source/Editor/UnrealEd/Private/AssetSelection.cpp b/Engine/Source/Editor/UnrealEd/Private/AssetSelection.cpp
--- a/Engine/Source/Editor/UnrealEd/Private/AssetSelection.cpp
+++ b/Engine/Source/Editor/UnrealEd/Private/AssetSelection.cpp
@@ -20,7 +20,6 @@
 #include "Kismet2/ComponentEditorUtils.h"
 #include "Engine/Selection.h"
 #include "Editor.h"
-#include "Matinee/MatineeActor.h"
 #include "ScopedTransaction.h"
 
 #include "LevelUtils.h"
diff --git a/Engine/Source/Editor/UnrealEd/Private/EditorExporters.cpp b/Engine/Source/Editor/UnrealEd/Private/EditorExporters.cpp
--- a/Engine/Source/Editor/UnrealEd/Private/EditorExporters.cpp
+++ b/Engine/Source/Editor/UnrealEd/Private/EditorExporters.cpp
@@ -58,7 +58,6 @@
 #include "GameFramework/PhysicsVolume.h"
 #include "EngineUtils.h"
 #include "Editor.h"
-#include "MatineeExporter.h"
 #include "FbxExporter.h"
 #include "StaticMeshAttributes.h"
 #include "StaticMeshOperations.h"
diff --git a/Engine/Source/Editor/UnrealEd/Private/EditorServer.cpp b/Engine/Source/Editor/UnrealEd/Private/EditorServer.cpp
--- a/Engine/Source/Editor/UnrealEd/Private/EditorServer.cpp
+++ b/Engine/Source/Editor/UnrealEd/Private/EditorServer.cpp
@@ -116,7 +116,6 @@
 #include "LandscapeProxy.h"
 #include "Lightmass/PrecomputedVisibilityOverrideVolume.h"
 #include "Animation/AnimSet.h"
-#include "Matinee/InterpTrackAnimControl.h"
 #include "InstancedFoliageActor.h"
 #include "IMovieSceneCapture.h"
 #include "MovieSceneCaptureModule.h"
@@ -199,21 +198,6 @@ namespace
 		int32 Num;
 		int32 Size;
 	};
-
-	struct FAnimSequenceUsageInfo
-	{
-		FAnimSequenceUsageInfo( float InStartOffset, float InEndOffset, UInterpTrackAnimControl* InAnimControl, int32 InTrackKeyIndex )
-		: 	StartOffset( InStartOffset )
-		,	EndOffset( InEndOffset )
-		,	AnimControl( InAnimControl )
-		,	TrackKeyIndex( InTrackKeyIndex )
-		{}
-
-		float						StartOffset;
-		float						EndOffset;
-		UInterpTrackAnimControl*	AnimControl;
-		int32							TrackKeyIndex;
-	};
 }
 
 /**
diff --git a/Engine/Source/Editor/UnrealEd/Private/EditorTransaction.cpp b/Engine/Source/Editor/UnrealEd/Private/EditorTransaction.cpp
--- a/Engine/Source/Editor/UnrealEd/Private/EditorTransaction.cpp
+++ b/Engine/Source/Editor/UnrealEd/Private/EditorTransaction.cpp
@@ -153,7 +153,7 @@ void FTransaction::FObjectRecord::SerializeObject( FArchive& Ar )
 
 void FTransaction::FObjectRecord::Restore( FTransaction* Owner )
 {
-	// only used by FMatineeTransaction:
+	// @TODO now that Matinee is gone, it's worth investigating removing this
 	if( !bRestored )
 	{
 		bRestored = true;
diff --git a/Engine/Source/Editor/UnrealEd/Private/EditorViewportClient.cpp b/Engine/Source/Editor/UnrealEd/Private/EditorViewportClient.cpp
--- a/Engine/Source/Editor/UnrealEd/Private/EditorViewportClient.cpp
+++ b/Engine/Source/Editor/UnrealEd/Private/EditorViewportClient.cpp
@@ -3935,7 +3935,7 @@ void FEditorViewportClient::Draw(FViewport* InViewport, FCanvas* Canvas)
 
 	ViewFamily.bIsHDR = Viewport->IsHDRViewport();
 
-	if( ModeTools->GetActiveMode( FBuiltinEditorModes::EM_InterpEdit ) == 0 || !AllowsCinematicControl() )
+	if( !AllowsCinematicControl() )
 	{
 		if( !UseEngineShowFlags.Game )
 		{
diff --git a/Engine/Source/Editor/UnrealEd/Private/Factories/EditorFactories.cpp b/Engine/Source/Editor/UnrealEd/Private/Factories/EditorFactories.cpp
--- a/Engine/Source/Editor/UnrealEd/Private/Factories/EditorFactories.cpp
+++ b/Engine/Source/Editor/UnrealEd/Private/Factories/EditorFactories.cpp
@@ -150,7 +150,6 @@
 #include "Engine/LevelScriptActor.h"
 #include "Engine/DataAsset.h"
 #include "Engine/BlueprintGeneratedClass.h"
-#include "Camera/CameraAnim.h"
 #include "Curves/CurveLinearColor.h"
 #include "Curves/CurveVector.h"
 #include "Engine/DataTable.h"
