Template Sequence: Remove ensure for sequence

diff --git a/Engine/Plugins/MovieScene/TemplateSequence/Source/TemplateSequence/Private/Section/TemplateSequenceSection.cpp b/Engine/Plugins/MovieScene/TemplateSequence/Source/TemplateSequence/Private/Section/TemplateSequenceSection.cpp
--- a/Engine/Plugins/MovieScene/TemplateSequence/Source/TemplateSequence/Private/Section/TemplateSequenceSection.cpp
+++ b/Engine/Plugins/MovieScene/TemplateSequence/Source/TemplateSequence/Private/Section/TemplateSequenceSection.cpp
@@ -107,7 +107,7 @@ void UTemplateSequenceSection::ImportEntityImpl(UMovieSceneEntitySystemLinker* E
 	using namespace UE::MovieScene;
 
 	UTemplateSequence* TemplateSubSequence = Cast<UTemplateSequence>(GetSequence());
-	if (!ensure(TemplateSubSequence))
+	if (!TemplateSubSequence)
 	{
 		return;
 	}
