Template Sequence: Check for valid subsequence before dereferencing

diff --git a/Engine/Plugins/MovieScene/TemplateSequence/Source/TemplateSequenceEditor/Private/TrackEditors/TemplateSequenceTrackEditor.cpp b/Engine/Plugins/MovieScene/TemplateSequence/Source/TemplateSequenceEditor/Private/TrackEditors/TemplateSequenceTrackEditor.cpp
--- a/Engine/Plugins/MovieScene/TemplateSequence/Source/TemplateSequenceEditor/Private/TrackEditors/TemplateSequenceTrackEditor.cpp
+++ b/Engine/Plugins/MovieScene/TemplateSequence/Source/TemplateSequenceEditor/Private/TrackEditors/TemplateSequenceTrackEditor.cpp
@@ -575,7 +575,17 @@ TArray<FTemplateSequenceSection::FScalablePropertyInfo> FTemplateSequenceSection
 
 	const UTemplateSequenceSection& TemplateSequenceSection = GetSectionObjectAs<UTemplateSequenceSection>();
 	UMovieSceneSequence* SubSequence = TemplateSequenceSection.GetSequence();
+	if (!SubSequence)
+	{
+		return AllAnimatedProperties;
+	}
+
 	UMovieScene* SubMovieScene = SubSequence->GetMovieScene();
+	if (!SubMovieScene)
+	{
+		return AllAnimatedProperties;
+	}
+
 	for (const FMovieSceneBinding& Binding : SubMovieScene->GetBindings())
 	{
 		for (const UMovieSceneTrack* BindingTrack : Binding.GetTracks())
