Sequencer: include play rate when computing sequence shake pattern duration

diff --git a/Engine/Plugins/MovieScene/TemplateSequence/Source/TemplateSequence/Private/SequenceCameraShake.cpp b/Engine/Plugins/MovieScene/TemplateSequence/Source/TemplateSequence/Private/SequenceCameraShake.cpp
--- a/Engine/Plugins/MovieScene/TemplateSequence/Source/TemplateSequence/Private/SequenceCameraShake.cpp
+++ b/Engine/Plugins/MovieScene/TemplateSequence/Source/TemplateSequence/Private/SequenceCameraShake.cpp
@@ -42,16 +42,18 @@ void USequenceCameraShakePattern::GetShakePatternInfoImpl(FCameraShakeInfo& OutI
 	{
 		if (UMovieScene* MovieScene = Sequence->GetMovieScene())
 		{
+			const float ActualPlayRate = (PlayRate > 0.f) ? PlayRate : 1.f;
+
 			if (bRandomSegment)
 			{
-				OutInfo.Duration = FCameraShakeDuration(RandomSegmentDuration);
+				OutInfo.Duration = FCameraShakeDuration(RandomSegmentDuration / ActualPlayRate);
 			}
 			else
 			{
 				const FFrameRate TickResolution = MovieScene->GetTickResolution();
 				const TRange<FFrameNumber> PlaybackRange = MovieScene->GetPlaybackRange();
 				const float Duration = TickResolution.AsSeconds(PlaybackRange.Size<FFrameNumber>());
-				OutInfo.Duration = FCameraShakeDuration(Duration);
+				OutInfo.Duration = FCameraShakeDuration(Duration / ActualPlayRate);
 			}
 
 			OutInfo.BlendIn = BlendInTime;
