Cameras: optimize sequencer camera animations

- Turn off all systems we don't need anymore
- Replace the default object binding handling with a very cheap one that does only what's needed for our camera stand-in object
- No need for the bound scene component impersonator interface anymore, since we're handling scene component bindings ourselves

diff --git a/Engine/Plugins/MovieScene/TemplateSequence/Source/TemplateSequence/Private/CameraAnimationSequencePlayer.cpp b/Engine/Plugins/MovieScene/TemplateSequence/Source/TemplateSequence/Private/CameraAnimationSequencePlayer.cpp
--- a/Engine/Plugins/MovieScene/TemplateSequence/Source/TemplateSequence/Private/CameraAnimationSequencePlayer.cpp
+++ b/Engine/Plugins/MovieScene/TemplateSequence/Source/TemplateSequence/Private/CameraAnimationSequencePlayer.cpp
@@ -264,8 +264,7 @@ UMovieSceneEntitySystemLinker* UCameraAnimationSequencePlayer::ConstructEntitySy
 			return Linker;
 		}
 	}
-	UMovieSceneEntitySystemLinker* NewLinker = NewObject<UMovieSceneEntitySystemLinker>(GetTransientPackage());
-	NewLinker->SetLinkerRole(UE::MovieScene::EEntitySystemLinkerRole::Standalone);
+	UMovieSceneEntitySystemLinker* NewLinker = UCameraAnimationSequenceSubsystem::CreateLinker(GetTransientPackage(), TEXT("StandaloneCameraAnimationLinker"));
 	return NewLinker;
 }
 
diff --git a/Engine/Plugins/MovieScene/TemplateSequence/Source/TemplateSequence/Private/CameraAnimationSequenceSubsystem.cpp b/Engine/Plugins/MovieScene/TemplateSequence/Source/TemplateSequence/Private/CameraAnimationSequenceSubsystem.cpp
--- a/Engine/Plugins/MovieScene/TemplateSequence/Source/TemplateSequence/Private/CameraAnimationSequenceSubsystem.cpp
+++ b/Engine/Plugins/MovieScene/TemplateSequence/Source/TemplateSequence/Private/CameraAnimationSequenceSubsystem.cpp
@@ -1,11 +1,169 @@
 // Copyright Epic Games, Inc. All Rights Reserved.
 
 #include "CameraAnimationSequenceSubsystem.h"
+#include "Components/SceneComponent.h"
 #include "Engine/World.h"
+#include "EntitySystem/BuiltInComponentTypes.h"
+#include "EntitySystem/MovieSceneEntityManager.h"
+#include "EntitySystem/MovieSceneEntityMutations.h"
 #include "EntitySystem/MovieSceneEntitySystemLinker.h"
+#include "EntitySystem/MovieSceneEntitySystemTask.h"
+#include "EntitySystem/MovieSceneMasterInstantiatorSystem.h"
+#include "EntitySystem/MovieSceneBoundSceneComponentInstantiator.h"
+#include "IMovieScenePlayer.h"
+#include "MovieSceneCommonHelpers.h"
+#include "Systems/MovieScenePropertyInstantiator.h"
 
 #define LOCTEXT_NAMESPACE "CameraAnimationSequenceSubsystem"
 
+namespace UE
+{
+namespace MovieScene
+{
+
+struct FCameraAnimationInstantiationMutation : IMovieSceneEntityMutation
+{
+	FInstanceRegistry& InstanceRegistry;
+	FBuiltInComponentTypes* BuiltInComponents;
+	FCameraAnimationInstantiationMutation(FInstanceRegistry& InInstanceRegistry)
+		: InstanceRegistry(InInstanceRegistry)
+	{
+		BuiltInComponents = FBuiltInComponentTypes::Get();
+	}
+	virtual void CreateMutation(FEntityManager* EntityManager, FComponentMask* InOutEntityComponentTypes) const
+	{
+		// Create all output components, a add a bound object on all entities with spawnable or object bindings.
+		FComponentRegistry* ComponentRegistry = EntityManager->GetComponents();
+		ComponentRegistry->Factories.ComputeMutuallyInclusiveComponents(*InOutEntityComponentTypes);
+		InOutEntityComponentTypes->Set(BuiltInComponents->BoundObject);
+	}
+	virtual void InitializeAllocation(FEntityAllocation* Allocation, const FComponentMask& AllocationType) const
+	{
+		const int32 Num = Allocation->Num();
+		FEntityAllocationWriteContext WriteContext = FEntityAllocationWriteContext::NewAllocation();
+		TComponentReader<FInstanceHandle> InstanceHandles = Allocation->ReadComponents(BuiltInComponents->InstanceHandle);
+		TComponentWriter<UObject*> OutBoundObjects = Allocation->WriteComponents(BuiltInComponents->BoundObject, WriteContext);
+	
+		if (AllocationType.Contains(BuiltInComponents->SpawnableBinding))
+		{
+			// Initialize spawned objects.
+			TComponentReader<FGuid> SpawnableBindings = Allocation->ReadComponents(BuiltInComponents->SpawnableBinding);
+			for (int32 Index = 0; Index < Num; ++Index)
+			{
+				SpawnObjectImpl(InstanceHandles[Index], SpawnableBindings[Index], OutBoundObjects[Index]);
+			}
+		}
+		else if (AllocationType.Contains(BuiltInComponents->GenericObjectBinding))
+		{
+			// Initialize bound objects.
+			TComponentReader<FGuid> ObjectBindings = Allocation->ReadComponents(BuiltInComponents->GenericObjectBinding);
+			for (int32 Index = 0; Index < Num; ++Index)
+			{
+				BindObjectImpl(InstanceHandles[Index], ObjectBindings[Index], OutBoundObjects[Index]);
+			}
+		}
+		else if (AllocationType.Contains(BuiltInComponents->SceneComponentBinding))
+		{
+			// Initialize bound scene components.
+			TComponentReader<FGuid> SceneComponentBindings = Allocation->ReadComponents(BuiltInComponents->SceneComponentBinding);
+			for (int32 Index = 0; Index < Num; ++Index)
+			{
+				BindObjectImpl(InstanceHandles[Index], SceneComponentBindings[Index], OutBoundObjects[Index]);
+			}
+		}
+	}
+	void SpawnObjectImpl(const FInstanceHandle& InstanceHandle, const FGuid& SpawnableBinding, UObject*& OutBoundObject) const
+	{
+		// We won't actually be spawning anything, because our player's spawn register will simply 
+		// return the fake camera "stand-in" object.
+		const FSequenceInstance& Instance = InstanceRegistry.GetInstance(InstanceHandle);
+		IMovieScenePlayer* Player = Instance.GetPlayer();
+		const UMovieSceneSequence* Sequence = Player->State.FindSequence(Instance.GetSequenceID());
+		UObject* SpawnedObject = Player->GetSpawnRegister().SpawnObject(
+			SpawnableBinding, *Sequence->GetMovieScene(), Instance.GetSequenceID(), *Player);
+		if (ensure(SpawnedObject))
+		{
+			OutBoundObject = SpawnedObject;
+		}
+	}
+	void BindObjectImpl(const FInstanceHandle& InstanceHandle, const FGuid& ObjectBinding, UObject*& OutBoundObject) const
+	{
+		const FSequenceInstance& Instance = InstanceRegistry.GetInstance(InstanceHandle);
+		IMovieScenePlayer* Player = Instance.GetPlayer();
+		TArrayView<TWeakObjectPtr<>> BoundObjects = Player->FindBoundObjects(ObjectBinding, Instance.GetSequenceID());
+		if (ensure(BoundObjects.Num() > 0))
+		{
+			// In theory we should get the scene component from the object, but we know that camera animations are
+			// always played on a camera animation player's "stand-in" camera object, which isn't really a camera, and
+			// certainly isn't an actor. It has its transform information directly on itself.
+			OutBoundObject = BoundObjects[0].Get();
+		}
+	}
+};
+
+} // namespace MovieScene
+} // namespace UE
+
+UCameraAnimationBoundObjectInstantiator::UCameraAnimationBoundObjectInstantiator(const FObjectInitializer& ObjInit)
+	: Super(ObjInit)
+{
+	using namespace UE::MovieScene;
+
+	Phase = UE::MovieScene::ESystemPhase::Instantiation;
+	SystemCategories = UCameraAnimationSequenceSubsystem::GetCameraAnimationSystemCategory();
+
+	if (HasAnyFlags(RF_ClassDefaultObject))
+	{
+		FBuiltInComponentTypes* BuiltInComponents = FBuiltInComponentTypes::Get();
+		DefineComponentProducer(GetClass(), BuiltInComponents->BoundObject);
+	}
+}
+
+void UCameraAnimationBoundObjectInstantiator::OnRun(FSystemTaskPrerequisites& InPrerequisites, FSystemSubsequentTasks& Subsequents)
+{
+	using namespace UE::MovieScene;
+
+	FBuiltInComponentTypes* BuiltInComponents = FBuiltInComponentTypes::Get();
+	FComponentRegistry* ComponentRegistry = Linker->EntityManager.GetComponents();
+
+	// Initialize all new allocations with bound objects and output components.
+	FCameraAnimationInstantiationMutation Mutation(*Linker->GetInstanceRegistry());
+	FEntityComponentFilter Filter = FEntityComponentFilter()
+		.Any({ BuiltInComponents->GenericObjectBinding, 
+				BuiltInComponents->SceneComponentBinding, 
+				BuiltInComponents->SpawnableBinding })
+		.All({ BuiltInComponents->InstanceHandle, BuiltInComponents->Tags.NeedsLink })
+		.None({ BuiltInComponents->Tags.NeedsUnlink });
+	Linker->EntityManager.MutateAll(Filter, Mutation);
+}
+
+UCameraAnimationEntitySystemLinker::UCameraAnimationEntitySystemLinker(const FObjectInitializer& ObjInit)
+	: Super(ObjInit)
+{
+	using namespace UE::MovieScene;
+
+	AutoLinkMode = EAutoLinkRelevantSystems::Disable;
+	SetLinkerRole(EEntitySystemLinkerRole::CameraAnimations);
+	GetSystemFilter().SetAllowedCategories(
+			// Eval time system, hierarchical bias systems, etc.
+			EEntitySystemCategory::Core |
+			// Systems to evaluate, blend, and set properties on bound objects
+			EEntitySystemCategory::ChannelEvaluators |
+			EEntitySystemCategory::BlenderSystems |
+			EEntitySystemCategory::PropertySystems |
+			// Our custom systems that avoid all the bound object instantiation that duplicates entities
+			UCameraAnimationSequenceSubsystem::GetCameraAnimationSystemCategory());
+	GetSystemFilter().AllowSystem(UMovieScenePropertyInstantiatorSystem::StaticClass());
+}
+
+void UCameraAnimationEntitySystemLinker::LinkRequiredSystems()
+{
+	using namespace UE::MovieScene;
+
+	// Link everything we want once and for all.
+	UMovieSceneEntitySystem::LinkAllSystems(this);
+}
+
 UCameraAnimationSequenceSubsystem* UCameraAnimationSequenceSubsystem::GetCameraAnimationSequenceSubsystem(const UWorld* InWorld)
 {
 	if (InWorld)
@@ -16,6 +174,20 @@ UCameraAnimationSequenceSubsystem* UCameraAnimationSequenceSubsystem::GetCameraA
 	return nullptr;
 }
 
+UMovieSceneEntitySystemLinker* UCameraAnimationSequenceSubsystem::CreateLinker(UObject* Outer, FName Name)
+{
+	UCameraAnimationEntitySystemLinker* NewLinker = NewObject<UCameraAnimationEntitySystemLinker>(Outer, Name);
+	NewLinker->LinkRequiredSystems();
+	return NewLinker;
+}
+
+UE::MovieScene::EEntitySystemCategory UCameraAnimationSequenceSubsystem::GetCameraAnimationSystemCategory()
+{
+	using namespace UE::MovieScene;
+	static EEntitySystemCategory CameraAnimationCategory = UMovieSceneEntitySystem::RegisterCustomSystemCategory();
+	return CameraAnimationCategory;
+}
+
 UCameraAnimationSequenceSubsystem::UCameraAnimationSequenceSubsystem()
 {
 }
@@ -41,11 +213,11 @@ UMovieSceneEntitySystemLinker* UCameraAnimationSequenceSubsystem::GetLinker(bool
 {
 	if (!Linker && bAutoCreate)
 	{
-		Linker = NewObject<UMovieSceneEntitySystemLinker>(this, TEXT("CameraAnimationSequenceSubsystemLinker"));
-		Linker->SetLinkerRole(UE::MovieScene::EEntitySystemLinkerRole::CameraAnimations);
+		Linker = CreateLinker(this, TEXT("CameraAnimationSequenceSubsystemLinker"));
 		Runner.AttachToLinker(Linker);
 	}
 	return Linker;
 }
 
 #undef LOCTEXT_NAMESPACE
+
diff --git a/Engine/Plugins/MovieScene/TemplateSequence/Source/TemplateSequence/Public/CameraAnimationSequencePlayer.h b/Engine/Plugins/MovieScene/TemplateSequence/Source/TemplateSequence/Public/CameraAnimationSequencePlayer.h
--- a/Engine/Plugins/MovieScene/TemplateSequence/Source/TemplateSequence/Public/CameraAnimationSequencePlayer.h
+++ b/Engine/Plugins/MovieScene/TemplateSequence/Source/TemplateSequence/Public/CameraAnimationSequencePlayer.h
@@ -4,7 +4,6 @@
 
 #include "CineCameraComponent.h"
 #include "Engine/Scene.h"
-#include "EntitySystem/MovieSceneBoundSceneComponentInstantiator.h"
 #include "TemplateSequencePlayer.h"
 #include "CameraAnimationSequencePlayer.generated.h"
 
@@ -12,7 +11,7 @@
  * A dummy class that we give to a sequence in lieu of an actual camera actor.
  */
 UCLASS()
-class TEMPLATESEQUENCE_API UCameraAnimationSequenceCameraStandIn : public UObject, public IMovieSceneSceneComponentImpersonator
+class TEMPLATESEQUENCE_API UCameraAnimationSequenceCameraStandIn : public UObject
 {
 public:
 
diff --git a/Engine/Plugins/MovieScene/TemplateSequence/Source/TemplateSequence/Public/CameraAnimationSequenceSubsystem.h b/Engine/Plugins/MovieScene/TemplateSequence/Source/TemplateSequence/Public/CameraAnimationSequenceSubsystem.h
--- a/Engine/Plugins/MovieScene/TemplateSequence/Source/TemplateSequence/Public/CameraAnimationSequenceSubsystem.h
+++ b/Engine/Plugins/MovieScene/TemplateSequence/Source/TemplateSequence/Public/CameraAnimationSequenceSubsystem.h
@@ -3,11 +3,34 @@
 #pragma once
 
 #include "CoreMinimal.h"
+#include "EntitySystem/MovieSceneEntityInstantiatorSystem.h"
 #include "EntitySystem/MovieSceneEntitySystemRunner.h"
 #include "Subsystems/WorldSubsystem.h"
 #include "UObject/ObjectMacros.h"
 #include "CameraAnimationSequenceSubsystem.generated.h"
 
+UCLASS()
+class UCameraAnimationBoundObjectInstantiator : public UMovieSceneEntityInstantiatorSystem
+{
+	GENERATED_BODY()
+
+public:
+	UCameraAnimationBoundObjectInstantiator(const FObjectInitializer& ObjInit);
+
+	void OnRun(FSystemTaskPrerequisites& InPrerequisites, FSystemSubsequentTasks& Subsequents);
+};
+
+UCLASS()
+class UCameraAnimationEntitySystemLinker : public UMovieSceneEntitySystemLinker
+{
+	GENERATED_BODY()
+
+public:
+	UCameraAnimationEntitySystemLinker(const FObjectInitializer& ObjInit);
+
+	void LinkRequiredSystems();
+};
+
 /**
  * World subsystem that holds global objects for handling camera animation sequences.
  */
@@ -20,6 +43,12 @@ public:
 	/** Get the camera animation sequence subsystem for the given world */
 	static UCameraAnimationSequenceSubsystem* GetCameraAnimationSequenceSubsystem(const UWorld* InWorld);
 
+	/** Create a new linker setup for handling camera animation sequences */
+	static UMovieSceneEntitySystemLinker* CreateLinker(UObject* Outer, FName Name);
+
+	/** Gets the system category for camera animation specific systems */
+	static UE::MovieScene::EEntitySystemCategory GetCameraAnimationSystemCategory();
+
 	/** Construct a new camera animation sequence subsystem */
 	UCameraAnimationSequenceSubsystem();
 	/** Destroy the camera animation sequence subsystem */
