CIS fix

diff --git a/Engine/Plugins/MovieScene/TemplateSequence/Source/TemplateSequence/Private/CameraAnimationSequencePlayer.cpp b/Engine/Plugins/MovieScene/TemplateSequence/Source/TemplateSequence/Private/CameraAnimationSequencePlayer.cpp
--- a/Engine/Plugins/MovieScene/TemplateSequence/Source/TemplateSequence/Private/CameraAnimationSequencePlayer.cpp
+++ b/Engine/Plugins/MovieScene/TemplateSequence/Source/TemplateSequence/Private/CameraAnimationSequencePlayer.cpp
@@ -5,7 +5,7 @@
 #include "Engine/World.h"
 #include "EntitySystem/MovieSceneEntitySystemLinker.h"
 #include "EntitySystem/MovieSceneEntitySystemTask.h"
-#include "EntitySystem/MovieScenePropertySystemTypes.inl"
+#include "EntitySystem/MovieScenePropertyComponentHandler.h"
 #include "GameFramework/WorldSettings.h"
 #include "CameraAnimationSequenceSubsystem.h"
 #include "MovieSceneTimeHelpers.h"
