Cameras: separate spawning from instantiation in the camera animation sequence systems

In some cases, the way ECS allocations are arranged in memory would mean that a bound object allocation would come before the spawnable allocation. This meant that these bound objects couldn't be resolved if they were children of the spawnable in the playing sequence.
We therefore can't do everything in one go when setting up the camera animation. We need to first do spawnables, and then the rest.

diff --git a/Engine/Plugins/MovieScene/TemplateSequence/Source/TemplateSequence/Private/CameraAnimationSequenceSubsystem.cpp b/Engine/Plugins/MovieScene/TemplateSequence/Source/TemplateSequence/Private/CameraAnimationSequenceSubsystem.cpp
--- a/Engine/Plugins/MovieScene/TemplateSequence/Source/TemplateSequence/Private/CameraAnimationSequenceSubsystem.cpp
+++ b/Engine/Plugins/MovieScene/TemplateSequence/Source/TemplateSequence/Private/CameraAnimationSequenceSubsystem.cpp
@@ -32,7 +32,7 @@ struct FCameraAnimationInstantiationMutation : IMovieSceneEntityMutation
 	}
 	virtual void CreateMutation(FEntityManager* EntityManager, FComponentMask* InOutEntityComponentTypes) const
 	{
-		// Create all output components, a add a bound object on all entities with spawnable or object bindings.
+		// Create all output components and bind objects/components.
 		FComponentRegistry* ComponentRegistry = EntityManager->GetComponents();
 		ComponentRegistry->Factories.ComputeMutuallyInclusiveComponents(*InOutEntityComponentTypes);
 		InOutEntityComponentTypes->Set(BuiltInComponents->BoundObject);
@@ -44,16 +44,7 @@ struct FCameraAnimationInstantiationMutation : IMovieSceneEntityMutation
 		TComponentReader<FInstanceHandle> InstanceHandles = Allocation->ReadComponents(BuiltInComponents->InstanceHandle);
 		TComponentWriter<UObject*> OutBoundObjects = Allocation->WriteComponents(BuiltInComponents->BoundObject, WriteContext);
 	
-		if (AllocationType.Contains(BuiltInComponents->SpawnableBinding))
-		{
-			// Initialize spawned objects.
-			TComponentReader<FGuid> SpawnableBindings = Allocation->ReadComponents(BuiltInComponents->SpawnableBinding);
-			for (int32 Index = 0; Index < Num; ++Index)
-			{
-				SpawnObjectImpl(InstanceHandles[Index], SpawnableBindings[Index], OutBoundObjects[Index]);
-			}
-		}
-		else if (AllocationType.Contains(BuiltInComponents->GenericObjectBinding))
+		if (AllocationType.Contains(BuiltInComponents->GenericObjectBinding))
 		{
 			// Initialize bound objects.
 			TComponentReader<FGuid> ObjectBindings = Allocation->ReadComponents(BuiltInComponents->GenericObjectBinding);
@@ -72,20 +63,6 @@ struct FCameraAnimationInstantiationMutation : IMovieSceneEntityMutation
 			}
 		}
 	}
-	void SpawnObjectImpl(const FInstanceHandle& InstanceHandle, const FGuid& SpawnableBinding, UObject*& OutBoundObject) const
-	{
-		// We won't actually be spawning anything, because our player's spawn register will simply 
-		// return the fake camera "stand-in" object.
-		const FSequenceInstance& Instance = InstanceRegistry.GetInstance(InstanceHandle);
-		IMovieScenePlayer* Player = Instance.GetPlayer();
-		const UMovieSceneSequence* Sequence = Player->State.FindSequence(Instance.GetSequenceID());
-		UObject* SpawnedObject = Player->GetSpawnRegister().SpawnObject(
-			SpawnableBinding, *Sequence->GetMovieScene(), Instance.GetSequenceID(), *Player);
-		if (ensure(SpawnedObject))
-		{
-			OutBoundObject = SpawnedObject;
-		}
-	}
 	void BindObjectImpl(const FInstanceHandle& InstanceHandle, const FGuid& ObjectBinding, UObject*& OutBoundObject) const
 	{
 		const FSequenceInstance& Instance = InstanceRegistry.GetInstance(InstanceHandle);
@@ -104,12 +81,45 @@ struct FCameraAnimationInstantiationMutation : IMovieSceneEntityMutation
 } // namespace MovieScene
 } // namespace UE
 
+UCameraAnimationSpawnableSystem::UCameraAnimationSpawnableSystem(const FObjectInitializer& ObjInit)
+	: Super(ObjInit)
+{
+	using namespace UE::MovieScene;
+
+	Phase = ESystemPhase::Spawn;
+	SystemCategories = UCameraAnimationSequenceSubsystem::GetCameraAnimationSystemCategory();
+}
+
+void UCameraAnimationSpawnableSystem::OnRun(FSystemTaskPrerequisites& InPrerequisites, FSystemSubsequentTasks& Subsequents)
+{
+	using namespace UE::MovieScene;
+
+	FInstanceRegistry* InstanceRegistry = Linker->GetInstanceRegistry();
+	FBuiltInComponentTypes* BuiltInComponents = FBuiltInComponentTypes::Get();
+
+	FEntityTaskBuilder()
+	.Read(BuiltInComponents->InstanceHandle)
+	.Read(BuiltInComponents->SpawnableBinding)
+	.FilterAll({ BuiltInComponents->Tags.NeedsLink })
+	.Iterate_PerEntity(&Linker->EntityManager, [InstanceRegistry](const FInstanceHandle& InstanceHandle, const FGuid& SpawnableBinding)
+		{
+			// We won't actually be spawning anything, because our player's spawn register will simply 
+			// return the fake camera "stand-in" object.
+			const FSequenceInstance& Instance = InstanceRegistry->GetInstance(InstanceHandle);
+			IMovieScenePlayer* Player = Instance.GetPlayer();
+			const UMovieSceneSequence* Sequence = Player->State.FindSequence(Instance.GetSequenceID());
+			UObject* SpawnedObject = Player->GetSpawnRegister().SpawnObject(
+				SpawnableBinding, *Sequence->GetMovieScene(), Instance.GetSequenceID(), *Player);
+			ensure(SpawnedObject);
+		});
+}
+
 UCameraAnimationBoundObjectInstantiator::UCameraAnimationBoundObjectInstantiator(const FObjectInitializer& ObjInit)
 	: Super(ObjInit)
 {
 	using namespace UE::MovieScene;
 
-	Phase = UE::MovieScene::ESystemPhase::Instantiation;
+	Phase = ESystemPhase::Instantiation;
 	SystemCategories = UCameraAnimationSequenceSubsystem::GetCameraAnimationSystemCategory();
 
 	if (HasAnyFlags(RF_ClassDefaultObject))
@@ -124,14 +134,12 @@ void UCameraAnimationBoundObjectInstantiator::OnRun(FSystemTaskPrerequisites& In
 	using namespace UE::MovieScene;
 
 	FBuiltInComponentTypes* BuiltInComponents = FBuiltInComponentTypes::Get();
-	FComponentRegistry* ComponentRegistry = Linker->EntityManager.GetComponents();
 
 	// Initialize all new allocations with bound objects and output components.
 	FCameraAnimationInstantiationMutation Mutation(*Linker->GetInstanceRegistry());
 	FEntityComponentFilter Filter = FEntityComponentFilter()
 		.Any({ BuiltInComponents->GenericObjectBinding, 
-				BuiltInComponents->SceneComponentBinding, 
-				BuiltInComponents->SpawnableBinding })
+				BuiltInComponents->SceneComponentBinding })
 		.All({ BuiltInComponents->InstanceHandle, BuiltInComponents->Tags.NeedsLink })
 		.None({ BuiltInComponents->Tags.NeedsUnlink });
 	Linker->EntityManager.MutateAll(Filter, Mutation);
@@ -153,7 +161,15 @@ UCameraAnimationEntitySystemLinker::UCameraAnimationEntitySystemLinker(const FOb
 			EEntitySystemCategory::PropertySystems |
 			// Our custom systems that avoid all the bound object instantiation that duplicates entities
 			UCameraAnimationSequenceSubsystem::GetCameraAnimationSystemCategory());
+	// Add the property instantiator specifically
 	GetSystemFilter().AllowSystem(UMovieScenePropertyInstantiatorSystem::StaticClass());
+
+	// NOTE: since we are tightly controlling the systems that are running to animate cameras, we may end
+	//       up not supporting new track types or custom features. However, 95% of these extra systems are
+	//       expected be part of the common categories that people might want to grow (custom channel 
+	//       evaluators, custom blending modes, new property types, etc)
+	//       For the last 5%, these custom systems can opt-in to run in camera animations by belonging to
+	//       the camera animation system category.
 }
 
 void UCameraAnimationEntitySystemLinker::LinkRequiredSystems()
diff --git a/Engine/Plugins/MovieScene/TemplateSequence/Source/TemplateSequence/Public/CameraAnimationSequenceSubsystem.h b/Engine/Plugins/MovieScene/TemplateSequence/Source/TemplateSequence/Public/CameraAnimationSequenceSubsystem.h
--- a/Engine/Plugins/MovieScene/TemplateSequence/Source/TemplateSequence/Public/CameraAnimationSequenceSubsystem.h
+++ b/Engine/Plugins/MovieScene/TemplateSequence/Source/TemplateSequence/Public/CameraAnimationSequenceSubsystem.h
@@ -9,6 +9,21 @@
 #include "UObject/ObjectMacros.h"
 #include "CameraAnimationSequenceSubsystem.generated.h"
 
+UCLASS()
+class UCameraAnimationSpawnableSystem : public UMovieSceneEntitySystem
+{
+	GENERATED_BODY()
+
+public:
+	UCameraAnimationSpawnableSystem(const FObjectInitializer& ObjInit);
+
+	void OnRun(FSystemTaskPrerequisites& InPrerequisites, FSystemSubsequentTasks& Subsequents);
+
+private:
+	// Always make this system irrelevant, so it's never used outside of this camera animation system
+	virtual bool IsRelevantImpl(UMovieSceneEntitySystemLinker* InLinker) const override { return false; }
+};
+
 UCLASS()
 class UCameraAnimationBoundObjectInstantiator : public UMovieSceneEntityInstantiatorSystem
 {
@@ -18,6 +33,10 @@ public:
 	UCameraAnimationBoundObjectInstantiator(const FObjectInitializer& ObjInit);
 
 	void OnRun(FSystemTaskPrerequisites& InPrerequisites, FSystemSubsequentTasks& Subsequents);
+
+private:
+	// Always make this system irrelevant, so it's never used outside of this camera animation system
+	virtual bool IsRelevantImpl(UMovieSceneEntitySystemLinker* InLinker) const override { return false; }
 };
 
 UCLASS()
