Template Sequence: Toolbar extender and icon fixes

diff --git a/Engine/Plugins/MovieScene/TemplateSequence/Source/TemplateSequenceEditor/Private/Customizations/CameraAnimationSequenceCustomization.cpp b/Engine/Plugins/MovieScene/TemplateSequence/Source/TemplateSequenceEditor/Private/Customizations/CameraAnimationSequenceCustomization.cpp
--- a/Engine/Plugins/MovieScene/TemplateSequence/Source/TemplateSequenceEditor/Private/Customizations/CameraAnimationSequenceCustomization.cpp
+++ b/Engine/Plugins/MovieScene/TemplateSequence/Source/TemplateSequenceEditor/Private/Customizations/CameraAnimationSequenceCustomization.cpp
@@ -17,7 +17,7 @@ void FCameraAnimationSequenceCustomization::RegisterSequencerCustomization(FSequ
 	FSequencerCustomizationInfo Customization;
 
 	TSharedRef<FExtender> ToolbarExtender = MakeShared<FExtender>();
-	ToolbarExtender->AddToolBarExtension("Base Commands", EExtensionHook::After, nullptr, FToolBarExtensionDelegate::CreateRaw(this, &FCameraAnimationSequenceCustomization::ExtendSequencerToolbar));
+	ToolbarExtender->AddToolBarExtension("BaseCommands", EExtensionHook::After, nullptr, FToolBarExtensionDelegate::CreateRaw(this, &FCameraAnimationSequenceCustomization::ExtendSequencerToolbar));
 	Customization.ToolbarExtender = ToolbarExtender;
 
 	Customization.OnAssetsDrop.BindLambda([](const TArray<UObject*>&, const FAssetDragDropOp&) -> ESequencerDropResult { return ESequencerDropResult::DropDenied; });
@@ -34,6 +34,8 @@ void FCameraAnimationSequenceCustomization::UnregisterSequencerCustomization()
 
 void FCameraAnimationSequenceCustomization::ExtendSequencerToolbar(FToolBarBuilder& ToolbarBuilder)
 {
+	ToolbarBuilder.BeginStyleOverride("SequencerToolBar");
+
 	ToolbarBuilder.AddSeparator();
 
 	ToolbarBuilder.AddComboButton(
@@ -42,6 +44,8 @@ void FCameraAnimationSequenceCustomization::ExtendSequencerToolbar(FToolBarBuild
 			LOCTEXT("CameraTypePicker", "Camera Type"),
 			LOCTEXT("CameraTypePickerTooltip", "Change the base camera actor type that this template sequence can bind to"),
 			FSlateIcon(FTemplateSequenceEditorStyle::Get()->GetStyleSetName(), "TemplateSequenceEditor.Chain"));
+	
+	ToolbarBuilder.EndStyleOverride();
 }
 
 TSharedRef<SWidget> FCameraAnimationSequenceCustomization::GetBoundCameraClassMenuContent()
diff --git a/Engine/Plugins/MovieScene/TemplateSequence/Source/TemplateSequenceEditor/Private/Customizations/TemplateSequenceCustomization.cpp b/Engine/Plugins/MovieScene/TemplateSequence/Source/TemplateSequenceEditor/Private/Customizations/TemplateSequenceCustomization.cpp
--- a/Engine/Plugins/MovieScene/TemplateSequence/Source/TemplateSequenceEditor/Private/Customizations/TemplateSequenceCustomization.cpp
+++ b/Engine/Plugins/MovieScene/TemplateSequence/Source/TemplateSequenceEditor/Private/Customizations/TemplateSequenceCustomization.cpp
@@ -19,7 +19,7 @@ void FTemplateSequenceCustomization::RegisterSequencerCustomization(FSequencerCu
 	FSequencerCustomizationInfo Customization;
 
 	TSharedRef<FExtender> ToolbarExtender = MakeShared<FExtender>();
-	ToolbarExtender->AddToolBarExtension("Base Commands", EExtensionHook::After, nullptr, FToolBarExtensionDelegate::CreateRaw(this, &FTemplateSequenceCustomization::ExtendSequencerToolbar));
+	ToolbarExtender->AddToolBarExtension("BaseCommands", EExtensionHook::After, nullptr, FToolBarExtensionDelegate::CreateRaw(this, &FTemplateSequenceCustomization::ExtendSequencerToolbar));
 	Customization.ToolbarExtender = ToolbarExtender;
 
 	Customization.OnAssetsDrop.BindRaw(this, &FTemplateSequenceCustomization::OnSequencerAssetsDrop);
@@ -36,6 +36,8 @@ void FTemplateSequenceCustomization::UnregisterSequencerCustomization()
 
 void FTemplateSequenceCustomization::ExtendSequencerToolbar(FToolBarBuilder& ToolbarBuilder)
 {
+	ToolbarBuilder.BeginStyleOverride("SequencerToolBar");
+
 	ToolbarBuilder.AddSeparator();
 
 	ToolbarBuilder.AddComboButton(
@@ -44,6 +46,8 @@ void FTemplateSequenceCustomization::ExtendSequencerToolbar(FToolBarBuilder& Too
 			LOCTEXT("BoundActorClassPicker", "Bound Actor Class"),
 			LOCTEXT("BoundActorClassPickerTooltip", "Change the base actor type that this template sequence can bind to"),
 			FSlateIcon(FTemplateSequenceEditorStyle::Get()->GetStyleSetName(), "TemplateSequenceEditor.Chain"));
+
+	ToolbarBuilder.EndStyleOverride();
 }
 
 TSharedRef<SWidget> FTemplateSequenceCustomization::GetBoundActorClassMenuContent()
diff --git a/Engine/Plugins/MovieScene/TemplateSequence/Source/TemplateSequenceEditor/Private/Styles/TemplateSequenceEditorStyle.h b/Engine/Plugins/MovieScene/TemplateSequence/Source/TemplateSequenceEditor/Private/Styles/TemplateSequenceEditorStyle.h
--- a/Engine/Plugins/MovieScene/TemplateSequence/Source/TemplateSequenceEditor/Private/Styles/TemplateSequenceEditorStyle.h
+++ b/Engine/Plugins/MovieScene/TemplateSequence/Source/TemplateSequenceEditor/Private/Styles/TemplateSequenceEditorStyle.h
@@ -40,8 +40,6 @@ public:
 
 		// toolbar icons
 		Set("TemplateSequenceEditor.Chain", new IMAGE_BRUSH("Chain_16x", Icon16x16));
-		Set("TemplateSequenceEditor.Chain", new IMAGE_BRUSH("Chain_64x", Icon64x64));
-		Set("TemplateSequenceEditor.Chain", new IMAGE_BRUSH("Chain_48x", Icon48x48));
 		Set("TemplateSequenceEditor.Chain.Small", new IMAGE_BRUSH("Chain_24x", Icon24x24));
 
 		FSlateStyleRegistry::RegisterSlateStyle(*this);
