Cameras: Re-enable dynamic system linking on camera animations/shakes

It turns out there's a whole bunch of systems we don't need in 95% of the case, so this will save a bit on evaluation costs.

diff --git a/Engine/Plugins/MovieScene/TemplateSequence/Source/TemplateSequence/Private/CameraAnimationSequenceSubsystem.cpp b/Engine/Plugins/MovieScene/TemplateSequence/Source/TemplateSequence/Private/CameraAnimationSequenceSubsystem.cpp
--- a/Engine/Plugins/MovieScene/TemplateSequence/Source/TemplateSequence/Private/CameraAnimationSequenceSubsystem.cpp
+++ b/Engine/Plugins/MovieScene/TemplateSequence/Source/TemplateSequence/Private/CameraAnimationSequenceSubsystem.cpp
@@ -93,6 +93,12 @@ UCameraAnimationSpawnableSystem::UCameraAnimationSpawnableSystem(const FObjectIn
 	SystemCategories = UCameraAnimationSequenceSubsystem::GetCameraAnimationSystemCategory();
 }
 
+bool UCameraAnimationSpawnableSystem::IsRelevantImpl(UMovieSceneEntitySystemLinker* InLinker) const
+{
+	using namespace UE::MovieScene;
+	return InLinker->GetLinkerRole() == EEntitySystemLinkerRole::CameraAnimations;
+}
+
 void UCameraAnimationSpawnableSystem::OnRun(FSystemTaskPrerequisites& InPrerequisites, FSystemSubsequentTasks& Subsequents)
 {
 	using namespace UE::MovieScene;
@@ -132,6 +138,12 @@ UCameraAnimationBoundObjectInstantiator::UCameraAnimationBoundObjectInstantiator
 	}
 }
 
+bool UCameraAnimationBoundObjectInstantiator::IsRelevantImpl(UMovieSceneEntitySystemLinker* InLinker) const
+{
+	using namespace UE::MovieScene;
+	return InLinker->GetLinkerRole() == EEntitySystemLinkerRole::CameraAnimations;
+}
+
 void UCameraAnimationBoundObjectInstantiator::OnRun(FSystemTaskPrerequisites& InPrerequisites, FSystemSubsequentTasks& Subsequents)
 {
 	using namespace UE::MovieScene;
@@ -153,7 +165,6 @@ UCameraAnimationEntitySystemLinker::UCameraAnimationEntitySystemLinker(const FOb
 {
 	using namespace UE::MovieScene;
 
-	AutoLinkMode = EAutoLinkRelevantSystems::Disable;
 	SetLinkerRole(EEntitySystemLinkerRole::CameraAnimations);
 	GetSystemFilter().SetAllowedCategories(
 			// Eval time system, hierarchical bias systems, etc.
@@ -175,14 +186,6 @@ UCameraAnimationEntitySystemLinker::UCameraAnimationEntitySystemLinker(const FOb
 	//       the camera animation system category.
 }
 
-void UCameraAnimationEntitySystemLinker::LinkRequiredSystems()
-{
-	using namespace UE::MovieScene;
-
-	// Link everything we want once and for all.
-	UMovieSceneEntitySystem::LinkAllSystems(this);
-}
-
 UCameraAnimationSequenceSubsystem* UCameraAnimationSequenceSubsystem::GetCameraAnimationSequenceSubsystem(const UWorld* InWorld)
 {
 	if (InWorld)
@@ -195,9 +198,7 @@ UCameraAnimationSequenceSubsystem* UCameraAnimationSequenceSubsystem::GetCameraA
 
 UMovieSceneEntitySystemLinker* UCameraAnimationSequenceSubsystem::CreateLinker(UObject* Outer, FName Name)
 {
-	UCameraAnimationEntitySystemLinker* NewLinker = NewObject<UCameraAnimationEntitySystemLinker>(Outer, Name);
-	NewLinker->LinkRequiredSystems();
-	return NewLinker;
+	return NewObject<UCameraAnimationEntitySystemLinker>(Outer, Name);
 }
 
 UE::MovieScene::EEntitySystemCategory UCameraAnimationSequenceSubsystem::GetCameraAnimationSystemCategory()
diff --git a/Engine/Plugins/MovieScene/TemplateSequence/Source/TemplateSequence/Public/CameraAnimationSequenceSubsystem.h b/Engine/Plugins/MovieScene/TemplateSequence/Source/TemplateSequence/Public/CameraAnimationSequenceSubsystem.h
--- a/Engine/Plugins/MovieScene/TemplateSequence/Source/TemplateSequence/Public/CameraAnimationSequenceSubsystem.h
+++ b/Engine/Plugins/MovieScene/TemplateSequence/Source/TemplateSequence/Public/CameraAnimationSequenceSubsystem.h
@@ -22,8 +22,7 @@ public:
 	void OnRun(FSystemTaskPrerequisites& InPrerequisites, FSystemSubsequentTasks& Subsequents);
 
 private:
-	// Always make this system irrelevant, so it's never used outside of this camera animation system
-	virtual bool IsRelevantImpl(UMovieSceneEntitySystemLinker* InLinker) const override { return false; }
+	virtual bool IsRelevantImpl(UMovieSceneEntitySystemLinker* InLinker) const;
 };
 
 UCLASS()
@@ -37,8 +36,7 @@ public:
 	void OnRun(FSystemTaskPrerequisites& InPrerequisites, FSystemSubsequentTasks& Subsequents);
 
 private:
-	// Always make this system irrelevant, so it's never used outside of this camera animation system
-	virtual bool IsRelevantImpl(UMovieSceneEntitySystemLinker* InLinker) const override { return false; }
+	virtual bool IsRelevantImpl(UMovieSceneEntitySystemLinker* InLinker) const;
 };
 
 UCLASS()
@@ -48,8 +46,6 @@ class UCameraAnimationEntitySystemLinker : public UMovieSceneEntitySystemLinker
 
 public:
 	UCameraAnimationEntitySystemLinker(const FObjectInitializer& ObjInit);
-
-	void LinkRequiredSystems();
 };
 
 /**
