Sequencer: only show template sequence menus on actor bindings, not component bindings.

diff --git a/Engine/Plugins/MovieScene/TemplateSequence/Source/TemplateSequenceEditor/Private/TrackEditors/TemplateSequenceTrackEditor.cpp b/Engine/Plugins/MovieScene/TemplateSequence/Source/TemplateSequenceEditor/Private/TrackEditors/TemplateSequenceTrackEditor.cpp
--- a/Engine/Plugins/MovieScene/TemplateSequence/Source/TemplateSequenceEditor/Private/TrackEditors/TemplateSequenceTrackEditor.cpp
+++ b/Engine/Plugins/MovieScene/TemplateSequence/Source/TemplateSequenceEditor/Private/TrackEditors/TemplateSequenceTrackEditor.cpp
@@ -59,13 +59,16 @@ bool FTemplateSequenceTrackEditor::SupportsSequence(UMovieSceneSequence* InSeque
 
 void FTemplateSequenceTrackEditor::BuildObjectBindingTrackMenu(FMenuBuilder& MenuBuilder, const TArray<FGuid>& ObjectBindings, const UClass* ObjectClass)
 {
-	const FText SubMenuEntryText = LOCTEXT("AddTemplateSequence", "Template Sequence");
-	const FText SubMenuEntryTooltip = LOCTEXT("AddTemplateSequenceTooltip", "Adds a track that can play a template sequence asset using the parent binding.");
-	
-	MenuBuilder.AddSubMenu(
-		SubMenuEntryText, SubMenuEntryTooltip,
-		FNewMenuDelegate::CreateRaw(this, &FTemplateSequenceTrackEditor::AddTemplateSequenceAssetSubMenu, ObjectBindings, ObjectClass)
-	);
+	if (ObjectClass->IsChildOf(AActor::StaticClass()))
+	{
+		const FText SubMenuEntryText = LOCTEXT("AddTemplateSequence", "Template Sequence");
+		const FText SubMenuEntryTooltip = LOCTEXT("AddTemplateSequenceTooltip", "Adds a track that can play a template sequence asset using the parent binding.");
+		
+		MenuBuilder.AddSubMenu(
+			SubMenuEntryText, SubMenuEntryTooltip,
+			FNewMenuDelegate::CreateRaw(this, &FTemplateSequenceTrackEditor::AddTemplateSequenceAssetSubMenu, ObjectBindings, ObjectClass)
+		);
+	}
 }
 
 TSharedPtr<SWidget> FTemplateSequenceTrackEditor::BuildOutlinerEditWidget(const FGuid& ObjectBinding, UMovieSceneTrack* Track, const FBuildEditWidgetParams& Params)
