Sequencer: prevent infinite loop while trying to add a possessable to a template sequence

diff --git a/Engine/Plugins/MovieScene/TemplateSequence/Source/TemplateSequence/Private/TemplateSequence.cpp b/Engine/Plugins/MovieScene/TemplateSequence/Source/TemplateSequence/Private/TemplateSequence.cpp
--- a/Engine/Plugins/MovieScene/TemplateSequence/Source/TemplateSequence/Private/TemplateSequence.cpp
+++ b/Engine/Plugins/MovieScene/TemplateSequence/Source/TemplateSequence/Private/TemplateSequence.cpp
@@ -79,7 +79,25 @@ void UTemplateSequence::BindPossessableObject(const FGuid& ObjectId, UObject& Po
 
 bool UTemplateSequence::CanPossessObject(UObject& Object, UObject* InPlaybackContext) const
 {
-	return Object.IsA<AActor>() || Object.IsA<UActorComponent>();
+	UClass* ExpectedActorClass = BoundActorClass.Get();
+	if (!ExpectedActorClass)
+	{
+		return false;
+	}
+
+	UActorComponent* Component = Cast<UActorComponent>(&Object);
+	if (!Component)
+	{
+		return false;
+	}
+
+	AActor* Owner = Component->GetOwner();
+	if (!Owner || !Owner->IsA(ExpectedActorClass))
+	{
+		return false;
+	}
+
+	return true;
 }
 
 void UTemplateSequence::LocateBoundObjects(const FGuid& ObjectId, UObject* Context, TArray<UObject*, TInlineAllocator<1>>& OutObjects) const
