Sequencer: remove unused property in template sequences

diff --git a/Engine/Plugins/MovieScene/TemplateSequence/Source/TemplateSequence/Public/TemplateSequence.h b/Engine/Plugins/MovieScene/TemplateSequence/Source/TemplateSequence/Public/TemplateSequence.h
--- a/Engine/Plugins/MovieScene/TemplateSequence/Source/TemplateSequence/Public/TemplateSequence.h
+++ b/Engine/Plugins/MovieScene/TemplateSequence/Source/TemplateSequence/Public/TemplateSequence.h
@@ -66,9 +66,6 @@ public:
 	UPROPERTY()
 	TSoftClassPtr<AActor> BoundActorClass;
 
-	UPROPERTY()
-	TSoftObjectPtr<AActor> BoundPreviewActor;
-
 	UPROPERTY()
 	TMap<FGuid, FName> BoundActorComponents;
 };
