Template Sequence / Movie Pipeline: Convert to asset definitions

diff --git a/Engine/Plugins/MovieScene/MovieRenderPipeline/Source/MovieRenderPipelineEditor/MovieRenderPipelineEditor.Build.cs b/Engine/Plugins/MovieScene/MovieRenderPipeline/Source/MovieRenderPipelineEditor/MovieRenderPipelineEditor.Build.cs
--- a/Engine/Plugins/MovieScene/MovieRenderPipeline/Source/MovieRenderPipelineEditor/MovieRenderPipelineEditor.Build.cs
+++ b/Engine/Plugins/MovieScene/MovieRenderPipeline/Source/MovieRenderPipelineEditor/MovieRenderPipelineEditor.Build.cs
@@ -41,6 +41,7 @@ public class MovieRenderPipelineEditor : ModuleRules
 		PrivateDependencyModuleNames.AddRange(
 			new string[]
 			{
+				"AssetDefinition",
 				"CoreUObject",
 				"Engine",
 				"Kismet",
diff --git a/Engine/Plugins/MovieScene/MovieRenderPipeline/Source/MovieRenderPipelineEditor/Private/AssetDefinition_PipelinePrimaryConfig.cpp b/Engine/Plugins/MovieScene/MovieRenderPipeline/Source/MovieRenderPipelineEditor/Private/AssetDefinition_PipelinePrimaryConfig.cpp
new file mode 100644
--- /dev/null
+++ b/Engine/Plugins/MovieScene/MovieRenderPipeline/Source/MovieRenderPipelineEditor/Private/AssetDefinition_PipelinePrimaryConfig.cpp
@@ -0,0 +1,32 @@
+// Copyright Epic Games, Inc. All Rights Reserved.
+
+#include "AssetDefinition_PipelinePrimaryConfig.h"
+
+#include "MoviePipelineAssetEditor.h"
+#include "Framework/Notifications/NotificationManager.h"
+#include "MoviePipelineQueueSubsystem.h"
+#include "IMovieRenderPipelineEditorModule.h"
+#include "Widgets/Notifications/SNotificationList.h"
+
+#define LOCTEXT_NAMESPACE "AssetTypeActions"
+
+EAssetCommandResult UAssetDefinition_PipelinePrimaryConfig::OpenAssets(const FAssetOpenArgs& OpenArgs) const
+{
+	UAssetEditorSubsystem* AssetEditorSubsystem = GEditor->GetEditorSubsystem<UAssetEditorSubsystem>();
+	if (AssetEditorSubsystem)
+	{
+		UMoviePipelineAssetEditor* AssetEditor = NewObject<UMoviePipelineAssetEditor>(AssetEditorSubsystem, NAME_None, RF_Transient);
+		if (AssetEditor)
+		{
+			for (UMoviePipelinePrimaryConfig* PrimaryConfig : OpenArgs.LoadObjects<UMoviePipelinePrimaryConfig>())
+			{
+				AssetEditor->SetObjectToEdit(PrimaryConfig);
+				AssetEditor->Initialize();
+			}
+		}
+	}
+
+	return EAssetCommandResult::Handled;
+}
+
+#undef LOCTEXT_NAMESPACE
diff --git a/Engine/Plugins/MovieScene/MovieRenderPipeline/Source/MovieRenderPipelineEditor/Private/AssetDefinition_PipelinePrimaryConfig.h b/Engine/Plugins/MovieScene/MovieRenderPipeline/Source/MovieRenderPipelineEditor/Private/AssetDefinition_PipelinePrimaryConfig.h
new file mode 100644
--- /dev/null
+++ b/Engine/Plugins/MovieScene/MovieRenderPipeline/Source/MovieRenderPipelineEditor/Private/AssetDefinition_PipelinePrimaryConfig.h
@@ -0,0 +1,27 @@
+// Copyright Epic Games, Inc. All Rights Reserved.
+
+#pragma once
+
+#include "MoviePipelinePrimaryConfig.h"
+#include "AssetDefinitionDefault.h"
+
+#include "AssetDefinition_PipelinePrimaryConfig.generated.h"
+
+UCLASS()
+class UAssetDefinition_PipelinePrimaryConfig : public UAssetDefinitionDefault
+{
+	GENERATED_BODY()
+
+public:
+	// UAssetDefinition Begin
+	virtual FText GetAssetDisplayName() const override { return NSLOCTEXT("AssetTypeActions", "AssetTypeActions_PipelinePrimaryConfig", "Movie Pipeline Primary Config"); }
+	virtual FLinearColor GetAssetColor() const override { return FLinearColor(FColor(78, 40, 165)); }
+	virtual TSoftClassPtr<UObject> GetAssetClass() const override { return UMoviePipelinePrimaryConfig::StaticClass(); }
+	virtual TConstArrayView<FAssetCategoryPath> GetAssetCategories() const override
+	{
+		static const auto Categories = { EAssetCategoryPaths::Animation };
+		return Categories;
+	}
+	virtual EAssetCommandResult OpenAssets(const FAssetOpenArgs& OpenArgs) const override;
+	// UAssetDefinition End
+};
diff --git a/Engine/Plugins/MovieScene/MovieRenderPipeline/Source/MovieRenderPipelineEditor/Private/AssetTypeActions_PipelineConfigs.cpp b/Engine/Plugins/MovieScene/MovieRenderPipeline/Source/MovieRenderPipelineEditor/Private/AssetTypeActions_PipelineConfigs.cpp
deleted file mode 100644
diff --git a/Engine/Plugins/MovieScene/MovieRenderPipeline/Source/MovieRenderPipelineEditor/Private/AssetTypeActions_PipelineConfigs.h b/Engine/Plugins/MovieScene/MovieRenderPipeline/Source/MovieRenderPipelineEditor/Private/AssetTypeActions_PipelineConfigs.h
deleted file mode 100644
diff --git a/Engine/Plugins/MovieScene/MovieRenderPipeline/Source/MovieRenderPipelineEditor/Private/MovieRenderPipelineEditorModule.cpp b/Engine/Plugins/MovieScene/MovieRenderPipeline/Source/MovieRenderPipelineEditor/Private/MovieRenderPipelineEditorModule.cpp
--- a/Engine/Plugins/MovieScene/MovieRenderPipeline/Source/MovieRenderPipelineEditor/Private/MovieRenderPipelineEditorModule.cpp
+++ b/Engine/Plugins/MovieScene/MovieRenderPipeline/Source/MovieRenderPipelineEditor/Private/MovieRenderPipelineEditorModule.cpp
@@ -29,8 +29,6 @@
 #include "MoviePipelineEditorBlueprintLibrary.h"
 #include "MovieSceneSection.h"
 #include "Sections/MovieSceneCinematicShotSection.h"
-#include "AssetToolsModule.h"
-#include "AssetTypeActions_PipelineConfigs.h"
 
 #define LOCTEXT_NAMESPACE "FMovieRenderPipelineEditorModule"
 
@@ -198,30 +196,10 @@ void FMovieRenderPipelineEditorModule::StartupModule()
 	RegisterTabImpl();
 	RegisterSettings();
 	RegisterMovieRenderer();
-
-	IAssetTools& AssetTools = FModuleManager::LoadModuleChecked<FAssetToolsModule>("AssetTools").Get();
-	{
-		PrimaryConfigAssetActions = MakeShared<FAssetTypeActions_PipelinePrimaryConfig>();
-		AssetTools.RegisterAssetTypeActions(PrimaryConfigAssetActions.ToSharedRef());
-
-		ShotConfigAssetActions = MakeShared<FAssetTypeActions_PipelineShotConfig>();
-		AssetTools.RegisterAssetTypeActions(ShotConfigAssetActions.ToSharedRef());
-
-		QueueAssetActions = MakeShared<FAssetTypeActions_PipelineQueue>();
-		AssetTools.RegisterAssetTypeActions(QueueAssetActions.ToSharedRef());
-	}
 }
 
 void FMovieRenderPipelineEditorModule::ShutdownModule()
 {
-	FAssetToolsModule* AssetToolsModule = FModuleManager::GetModulePtr<FAssetToolsModule>("AssetTools");
-	if (AssetToolsModule)
-	{
-		AssetToolsModule->Get().UnregisterAssetTypeActions(PrimaryConfigAssetActions.ToSharedRef());
-		AssetToolsModule->Get().UnregisterAssetTypeActions(ShotConfigAssetActions.ToSharedRef());
-		AssetToolsModule->Get().UnregisterAssetTypeActions(QueueAssetActions.ToSharedRef());
-	}
-
 	UnregisterMovieRenderer();
 	UnregisterSettings();
 	FMoviePipelineCommands::Unregister();
diff --git a/Engine/Plugins/MovieScene/MovieRenderPipeline/Source/MovieRenderPipelineEditor/Private/MovieRenderPipelineEditorModule.h b/Engine/Plugins/MovieScene/MovieRenderPipeline/Source/MovieRenderPipelineEditor/Private/MovieRenderPipelineEditorModule.h
--- a/Engine/Plugins/MovieScene/MovieRenderPipeline/Source/MovieRenderPipelineEditor/Private/MovieRenderPipelineEditorModule.h
+++ b/Engine/Plugins/MovieScene/MovieRenderPipeline/Source/MovieRenderPipelineEditor/Private/MovieRenderPipelineEditorModule.h
@@ -6,10 +6,6 @@
 #include "IMovieRendererInterface.h"
 #include "Delegates/IDelegateInstance.h"
 
-class UMoviePipelineShotConfig;
-class UMoviePipelinePrimaryConfig;
-class UMoviePipelineExecutorBase;
-
 class FMovieRenderPipelineEditorModule : public IMovieRenderPipelineEditorModule
 {
 public:
@@ -24,8 +20,5 @@ private:
 	void RegisterMovieRenderer();
 	void UnregisterMovieRenderer();
 
-	TSharedPtr<class FAssetTypeActions_PipelinePrimaryConfig> PrimaryConfigAssetActions;
-	TSharedPtr<class FAssetTypeActions_PipelineShotConfig> ShotConfigAssetActions;
-	TSharedPtr<class FAssetTypeActions_PipelineQueue> QueueAssetActions;
 	FDelegateHandle MovieRendererDelegate;
 };
diff --git a/Engine/Plugins/MovieScene/TemplateSequence/Source/TemplateSequenceEditor/Private/AssetTools/AssetDefinition_CameraAnimationSequence.cpp b/Engine/Plugins/MovieScene/TemplateSequence/Source/TemplateSequenceEditor/Private/AssetTools/AssetDefinition_CameraAnimationSequence.cpp
new file mode 100644
--- /dev/null
+++ b/Engine/Plugins/MovieScene/TemplateSequence/Source/TemplateSequenceEditor/Private/AssetTools/AssetDefinition_CameraAnimationSequence.cpp
@@ -0,0 +1,15 @@
+// Copyright Epic Games, Inc. All Rights Reserved.
+
+#include "AssetDefinition_CameraAnimationSequence.h"
+#include "TemplateSequenceEditorToolkit.h"
+
+#define LOCTEXT_NAMESPACE "AssetTypeActions"
+
+void UAssetDefinition_CameraAnimationSequence::InitializeToolkitParams(FTemplateSequenceToolkitParams& ToolkitParams) const
+{
+    UAssetDefinition_TemplateSequence::InitializeToolkitParams(ToolkitParams);
+
+    ToolkitParams.bCanChangeBinding = false;
+}
+
+#undef LOCTEXT_NAMESPACE
diff --git a/Engine/Plugins/MovieScene/TemplateSequence/Source/TemplateSequenceEditor/Private/AssetTools/AssetDefinition_CameraAnimationSequence.h b/Engine/Plugins/MovieScene/TemplateSequence/Source/TemplateSequenceEditor/Private/AssetTools/AssetDefinition_CameraAnimationSequence.h
new file mode 100644
--- /dev/null
+++ b/Engine/Plugins/MovieScene/TemplateSequence/Source/TemplateSequenceEditor/Private/AssetTools/AssetDefinition_CameraAnimationSequence.h
@@ -0,0 +1,27 @@
+// Copyright Epic Games, Inc. All Rights Reserved.
+
+#pragma once
+
+#include "CameraAnimationSequence.h"
+#include "AssetDefinitionDefault.h"
+#include "AssetDefinition_TemplateSequence.h"
+
+#include "AssetDefinition_CameraAnimationSequence.generated.h"
+
+struct FTemplateSequenceToolkitParams;
+
+UCLASS()
+class UAssetDefinition_CameraAnimationSequence : public UAssetDefinition_TemplateSequence
+{
+	GENERATED_BODY()
+
+public:
+	// UAssetDefinition Begin
+	virtual FText GetAssetDisplayName() const override { return NSLOCTEXT("AssetTypeActions", "AssetTypeActions_CameraAnimationSequence", "Camera Animation Sequence"); }
+	virtual TSoftClassPtr<UObject> GetAssetClass() const override { return UCameraAnimationSequence::StaticClass(); }
+	// UAssetDefinition End
+
+protected:
+
+	virtual void InitializeToolkitParams(FTemplateSequenceToolkitParams& ToolkitParams) const override;
+};
diff --git a/Engine/Plugins/MovieScene/TemplateSequence/Source/TemplateSequenceEditor/Private/AssetTools/AssetDefinition_TemplateSequence.cpp b/Engine/Plugins/MovieScene/TemplateSequence/Source/TemplateSequenceEditor/Private/AssetTools/AssetDefinition_TemplateSequence.cpp
new file mode 100644
--- /dev/null
+++ b/Engine/Plugins/MovieScene/TemplateSequence/Source/TemplateSequenceEditor/Private/AssetTools/AssetDefinition_TemplateSequence.cpp
@@ -0,0 +1,42 @@
+// Copyright Epic Games, Inc. All Rights Reserved.
+
+#include "AssetDefinition_TemplateSequence.h"
+#include "TemplateSequenceEditorToolkit.h"
+#include "Styles/TemplateSequenceEditorStyle.h"
+
+#define LOCTEXT_NAMESPACE "AssetTypeActions"
+
+FAssetOpenSupport UAssetDefinition_TemplateSequence::GetAssetOpenSupport(const FAssetOpenSupportArgs& OpenSupportArgs) const
+{
+	return FAssetOpenSupport(OpenSupportArgs.OpenMethod,OpenSupportArgs.OpenMethod == EAssetOpenMethod::Edit, EToolkitMode::WorldCentric); 
+}
+
+EAssetCommandResult UAssetDefinition_TemplateSequence::OpenAssets(const FAssetOpenArgs& OpenArgs) const
+{
+	UWorld* WorldContext = nullptr;
+	for (const FWorldContext& Context : GEngine->GetWorldContexts())
+	{
+		if (Context.WorldType == EWorldType::Editor)
+		{
+			WorldContext = Context.World();
+			break;
+		}
+	}
+
+	if (!ensure(WorldContext))
+	{
+		return EAssetCommandResult::Handled;
+	}
+
+	for (UTemplateSequence* TemplateSequence : OpenArgs.LoadObjects<UTemplateSequence>())
+	{
+		TSharedRef<FTemplateSequenceEditorToolkit> Toolkit = MakeShareable(new FTemplateSequenceEditorToolkit(FTemplateSequenceEditorStyle::Get()));
+		FTemplateSequenceToolkitParams ToolkitParams;
+		InitializeToolkitParams(ToolkitParams);
+		Toolkit->Initialize(OpenArgs.GetToolkitMode(), OpenArgs.ToolkitHost, TemplateSequence, ToolkitParams);
+	}	
+
+	return EAssetCommandResult::Handled;
+}
+
+#undef LOCTEXT_NAMESPACE
diff --git a/Engine/Plugins/MovieScene/TemplateSequence/Source/TemplateSequenceEditor/Private/AssetTools/AssetDefinition_TemplateSequence.h b/Engine/Plugins/MovieScene/TemplateSequence/Source/TemplateSequenceEditor/Private/AssetTools/AssetDefinition_TemplateSequence.h
new file mode 100644
--- /dev/null
+++ b/Engine/Plugins/MovieScene/TemplateSequence/Source/TemplateSequenceEditor/Private/AssetTools/AssetDefinition_TemplateSequence.h
@@ -0,0 +1,34 @@
+// Copyright Epic Games, Inc. All Rights Reserved.
+
+#pragma once
+
+#include "TemplateSequence.h"
+#include "AssetDefinitionDefault.h"
+
+#include "AssetDefinition_TemplateSequence.generated.h"
+
+struct FTemplateSequenceToolkitParams;
+
+UCLASS()
+class UAssetDefinition_TemplateSequence : public UAssetDefinitionDefault
+{
+	GENERATED_BODY()
+
+public:
+	// UAssetDefinition Begin
+	virtual FText GetAssetDisplayName() const override { return NSLOCTEXT("AssetTypeActions", "AssetTypeActions_TemplateSequence", "Template Sequence"); }
+	virtual FLinearColor GetAssetColor() const override { return FLinearColor(FColor(200, 80, 80)); }
+	virtual TSoftClassPtr<UObject> GetAssetClass() const override { return UTemplateSequence::StaticClass(); }
+	virtual TConstArrayView<FAssetCategoryPath> GetAssetCategories() const override
+	{
+		static const auto Categories = { EAssetCategoryPaths::Cinematics };
+		return Categories;
+	}
+	virtual FAssetOpenSupport GetAssetOpenSupport(const FAssetOpenSupportArgs& OpenSupportArgs) const override;
+	virtual EAssetCommandResult OpenAssets(const FAssetOpenArgs& OpenArgs) const override;
+	// UAssetDefinition End
+
+protected:
+
+	virtual void InitializeToolkitParams(FTemplateSequenceToolkitParams& ToolkitParams) const {}
+};
diff --git a/Engine/Plugins/MovieScene/TemplateSequence/Source/TemplateSequenceEditor/Private/AssetTools/CameraAnimationSequenceActions.cpp b/Engine/Plugins/MovieScene/TemplateSequence/Source/TemplateSequenceEditor/Private/AssetTools/CameraAnimationSequenceActions.cpp
deleted file mode 100644
diff --git a/Engine/Plugins/MovieScene/TemplateSequence/Source/TemplateSequenceEditor/Private/AssetTools/CameraAnimationSequenceActions.h b/Engine/Plugins/MovieScene/TemplateSequence/Source/TemplateSequenceEditor/Private/AssetTools/CameraAnimationSequenceActions.h
deleted file mode 100644
diff --git a/Engine/Plugins/MovieScene/TemplateSequence/Source/TemplateSequenceEditor/Private/AssetTools/TemplateSequenceActions.cpp b/Engine/Plugins/MovieScene/TemplateSequence/Source/TemplateSequenceEditor/Private/AssetTools/TemplateSequenceActions.cpp
deleted file mode 100644
diff --git a/Engine/Plugins/MovieScene/TemplateSequence/Source/TemplateSequenceEditor/Private/AssetTools/TemplateSequenceActions.h b/Engine/Plugins/MovieScene/TemplateSequence/Source/TemplateSequenceEditor/Private/AssetTools/TemplateSequenceActions.h
deleted file mode 100644
diff --git a/Engine/Plugins/MovieScene/TemplateSequence/Source/TemplateSequenceEditor/Private/TemplateSequenceEditorModule.cpp b/Engine/Plugins/MovieScene/TemplateSequence/Source/TemplateSequenceEditor/Private/TemplateSequenceEditorModule.cpp
--- a/Engine/Plugins/MovieScene/TemplateSequence/Source/TemplateSequenceEditor/Private/TemplateSequenceEditorModule.cpp
+++ b/Engine/Plugins/MovieScene/TemplateSequence/Source/TemplateSequenceEditor/Private/TemplateSequenceEditorModule.cpp
@@ -1,7 +1,5 @@
 // Copyright Epic Games, Inc. All Rights Reserved.
 
-#include "AssetTools/CameraAnimationSequenceActions.h"
-#include "AssetTools/TemplateSequenceActions.h"
 #include "CameraAnimationSequence.h"
 #include "Customizations/CameraAnimationSequenceCustomization.h"
 #include "Customizations/TemplateSequenceCustomization.h"
@@ -10,6 +8,7 @@
 #include "ISettingsModule.h"
 #include "Misc/MovieSceneSequenceEditor_TemplateSequence.h"
 #include "Misc/TemplateSequenceEditorSettings.h"
+#include "Modules/ModuleManager.h"
 #include "SequencerSettings.h"
 #include "Styles/TemplateSequenceEditorStyle.h"
 #include "TrackEditors/TemplateSequenceTrackEditor.h"
@@ -30,7 +29,6 @@ public:
 	virtual void StartupModule() override
 	{
 		RegisterSettings();
-		RegisterAssetTools();
 		RegisterSequenceEditor();
 		RegisterTrackEditors();
 		RegisterSequenceCustomizations();
@@ -41,7 +39,6 @@ public:
 		UnregisterSequenceCustomizations();
 		UnregisterTrackEditors();
 		UnregisterSequenceEditor();
-		UnregisterAssetTools();
 		UnregisterSettings();
 	}
 
@@ -91,30 +88,6 @@ private:
 		}
 	}
 
-	void RegisterAssetTools()
-	{
-		IAssetTools& AssetTools = FModuleManager::LoadModuleChecked<FAssetToolsModule>("AssetTools").Get();
-		EAssetTypeCategories::Type CinematicAssetCategoryBit = AssetTools.RegisterAdvancedAssetCategory(FName(TEXT("Cinematics")), LOCTEXT("CinematicsAssetCategory", "Cinematics"));
-
-		TemplateSequenceTypeActions = MakeShared<FTemplateSequenceActions>(FTemplateSequenceEditorStyle::Get(), CinematicAssetCategoryBit);
-		AssetTools.RegisterAssetTypeActions(TemplateSequenceTypeActions.ToSharedRef());
-
-		CameraAnimationSequenceActions = MakeShared<FCameraAnimationSequenceActions>(FTemplateSequenceEditorStyle::Get(), CinematicAssetCategoryBit);
-		AssetTools.RegisterAssetTypeActions(CameraAnimationSequenceActions.ToSharedRef());
-	}
-
-	void UnregisterAssetTools()
-	{
-		FAssetToolsModule* AssetToolsModule = FModuleManager::GetModulePtr<FAssetToolsModule>("AssetTools");
-
-		if (AssetToolsModule != nullptr)
-		{
-			IAssetTools& AssetTools = AssetToolsModule->Get();
-			AssetTools.UnregisterAssetTypeActions(TemplateSequenceTypeActions.ToSharedRef());
-			AssetTools.UnregisterAssetTypeActions(CameraAnimationSequenceActions.ToSharedRef());
-		}
-	}
-
 	void RegisterSequenceEditor()
 	{
 		ISequencerModule& SequencerModule = FModuleManager::LoadModuleChecked<ISequencerModule>("Sequencer");
@@ -169,9 +142,6 @@ private:
 
 private:
 
-	TSharedPtr<FTemplateSequenceActions> TemplateSequenceTypeActions;
-	TSharedPtr<FCameraAnimationSequenceActions> CameraAnimationSequenceActions;
-
 	FDelegateHandle SequenceEditorHandle;
 	FDelegateHandle TemplateSequenceTrackCreateEditorHandle;
 
diff --git a/Engine/Plugins/MovieScene/TemplateSequence/Source/TemplateSequenceEditor/TemplateSequenceEditor.Build.cs b/Engine/Plugins/MovieScene/TemplateSequence/Source/TemplateSequenceEditor/TemplateSequenceEditor.Build.cs
--- a/Engine/Plugins/MovieScene/TemplateSequence/Source/TemplateSequenceEditor/TemplateSequenceEditor.Build.cs
+++ b/Engine/Plugins/MovieScene/TemplateSequence/Source/TemplateSequenceEditor/TemplateSequenceEditor.Build.cs
@@ -8,6 +8,7 @@ public class TemplateSequenceEditor : ModuleRules
 	{
 		PrivateDependencyModuleNames.AddRange(
 			new string[] {
+				"AssetDefinition",
 				"TemplateSequence",
 				"BlueprintGraph",
 				"CinematicCamera",
