Sequencer: fix time-related bugs in camera animations, and add support for a start offset

diff --git a/Engine/Plugins/Cameras/GameplayCameras/Source/GameplayCameras/Private/CameraAnimationCameraModifier.cpp b/Engine/Plugins/Cameras/GameplayCameras/Source/GameplayCameras/Private/CameraAnimationCameraModifier.cpp
--- a/Engine/Plugins/Cameras/GameplayCameras/Source/GameplayCameras/Private/CameraAnimationCameraModifier.cpp
+++ b/Engine/Plugins/Cameras/GameplayCameras/Source/GameplayCameras/Private/CameraAnimationCameraModifier.cpp
@@ -52,10 +52,6 @@ float UCameraAnimationCameraModifier::EvaluateEasing(ECameraAnimationEasingType
 {
 	using namespace CameraAnimationCameraModifierImpl;
 	
-	// Used for in-out easing
-	const float InTime = Interp * 2.f;
-	const float OutTime = (Interp - .5f) * 2.f;
-
 	switch (EasingType)
 	{
 	case ECameraAnimationEasingType::Sinusoidal: 	return Sinusoidal(Interp);
@@ -112,7 +108,7 @@ FCameraAnimationHandle UCameraAnimationCameraModifier::PlayCameraAnimation(UCame
 	NewCameraAnimation.Player->SetBoundObjectOverride(NewCameraAnimation.CameraStandIn);
 
 	// Initialize it and start playing.
-	NewCameraAnimation.Player->Initialize(Sequence);
+	NewCameraAnimation.Player->Initialize(Sequence, Params.StartOffset);
 	NewCameraAnimation.Player->Play(Params.bLoop, Params.bRandomStartTime);
 
 	return InstanceHandle;
@@ -215,12 +211,12 @@ void UCameraAnimationCameraModifier::DisplayDebug(class UCanvas* Canvas, const F
 
 		if (ActiveAnimation.IsValid())
 		{
-			const FFrameRate InputRate = ActiveAnimation.Player->GetInputRate();
-			const FFrameNumber DurationFrames = ActiveAnimation.Player->GetDuration();
+			const FFrameRate DisplayRate = ActiveAnimation.Player->GetInputRate();
+			const FFrameTime DurationFrames = ActiveAnimation.Player->GetDuration();
 			const FFrameTime CurrentPosition = ActiveAnimation.Player->GetCurrentPosition();
 
-			const float CurrentTime = InputRate.AsSeconds(CurrentPosition);
-			const float DurationSeconds = InputRate.AsSeconds(DurationFrames);
+			const float CurrentTime = DisplayRate.AsSeconds(CurrentPosition);
+			const float DurationSeconds = DisplayRate.AsSeconds(DurationFrames);
 
 			const FString LoopString = ActiveAnimation.Params.bLoop ? TEXT(" Looping") : TEXT("");
 			const FString EaseInString = ActiveAnimation.bIsEasingIn ? FString::Printf(TEXT(" Easing In: %f / %f"), ActiveAnimation.EaseInCurrentTime, ActiveAnimation.Params.EaseInDuration) : TEXT("");
@@ -321,15 +317,11 @@ void UCameraAnimationCameraModifier::TickAnimation(FActiveCameraAnimationInfo& C
 	UCameraAnimationSequencePlayer* Player = CameraAnimation.Player;
 	UCameraAnimationSequenceCameraStandIn* CameraStandIn = CameraAnimation.CameraStandIn;
 
-	const FFrameRate InputRate = Player->GetInputRate();
+	const FFrameRate DisplayRate = Player->GetInputRate();
 	const FFrameTime CurrentPosition = Player->GetCurrentPosition();
-	const float CurrentTime = InputRate.AsSeconds(CurrentPosition);
-	const float DurationTime = InputRate.AsSeconds(Player->GetDuration()) * Params.PlayRate;
 
 	const float ScaledDeltaTime = DeltaTime * Params.PlayRate;
-
-	const float NewTime = CurrentTime + ScaledDeltaTime;
-	const FFrameTime NewPosition = CurrentPosition + DeltaTime * Params.PlayRate * InputRate;
+	const FFrameTime NewPosition = CurrentPosition + ScaledDeltaTime * DisplayRate;
 
 	// Advance any easing times.
 	if (CameraAnimation.bIsEasingIn)
@@ -344,6 +336,8 @@ void UCameraAnimationCameraModifier::TickAnimation(FActiveCameraAnimationInfo& C
 	// Start easing out if we're nearing the end.
 	if (!Player->GetIsLooping())
 	{
+		const float NewTime = DisplayRate.AsSeconds(NewPosition);
+		const float DurationTime = DisplayRate.AsSeconds(Player->GetDuration()) * Params.PlayRate;
 		const float BlendOutStartTime = DurationTime - Params.EaseOutDuration;
 		if (NewTime > BlendOutStartTime)
 		{
diff --git a/Engine/Plugins/Cameras/GameplayCameras/Source/GameplayCameras/Public/CameraAnimationCameraModifier.h b/Engine/Plugins/Cameras/GameplayCameras/Source/GameplayCameras/Public/CameraAnimationCameraModifier.h
--- a/Engine/Plugins/Cameras/GameplayCameras/Source/GameplayCameras/Public/CameraAnimationCameraModifier.h
+++ b/Engine/Plugins/Cameras/GameplayCameras/Source/GameplayCameras/Public/CameraAnimationCameraModifier.h
@@ -67,6 +67,9 @@ struct FCameraAnimationParams
 	/** Whether the camera animation should loop */
 	UPROPERTY(BlueprintReadWrite, Category="Camera Animation")
 	bool bLoop = false;
+	/** Offset, in frames, into the animation to start at */
+	UPROPERTY(BlueprintReadWrite, Category="Camera Animation")
+	int32 StartOffset = 0;
 	/** Whether the camera animation should have a random start time */
 	UPROPERTY(BlueprintReadWrite, Category="Camera Animation")
 	bool bRandomStartTime = false;
diff --git a/Engine/Plugins/MovieScene/TemplateSequence/Source/TemplateSequence/Private/CameraAnimationSequencePlayer.cpp b/Engine/Plugins/MovieScene/TemplateSequence/Source/TemplateSequence/Private/CameraAnimationSequencePlayer.cpp
--- a/Engine/Plugins/MovieScene/TemplateSequence/Source/TemplateSequence/Private/CameraAnimationSequencePlayer.cpp
+++ b/Engine/Plugins/MovieScene/TemplateSequence/Source/TemplateSequence/Private/CameraAnimationSequencePlayer.cpp
@@ -11,6 +11,7 @@
 #include "GameFramework/WorldSettings.h"
 #include "CameraAnimationSequenceSubsystem.h"
 #include "MovieScene.h"
+#include "MovieSceneTimeHelpers.h"
 #include "MovieSceneTracksComponentTypes.h"
 #include "MovieSceneSpawnable.h"
 #include "TemplateSequence.h"
@@ -294,12 +295,17 @@ void UCameraAnimationSequencePlayer::SetBoundObjectOverride(UObject* InObject)
 	SpawnRegister.SetSpawnedObject(InObject);
 }
 
-FFrameNumber UCameraAnimationSequencePlayer::GetDuration() const
+FFrameTime UCameraAnimationSequencePlayer::GetDuration() const
 {
-	return ConvertFrameTime(DurationFrames, PlayPosition.GetOutputRate(), PlayPosition.GetInputRate()).FloorToFrame();
+	return DurationFrames;
 }
 
-void UCameraAnimationSequencePlayer::Initialize(UMovieSceneSequence* InSequence)
+FFrameTime UCameraAnimationSequencePlayer::GetCurrentPosition() const
+{
+	return PlayPosition.GetCurrentPosition();
+}
+
+void UCameraAnimationSequencePlayer::Initialize(UMovieSceneSequence* InSequence, int32 StartOffset)
 {
 	checkf(InSequence, TEXT("Invalid sequence given to player"));
 	
@@ -320,13 +326,29 @@ void UCameraAnimationSequencePlayer::Initialize(UMovieSceneSequence* InSequence)
 		PlayPosition.SetTimeBase(DisplayRate, TickResolution, EvaluationType);
 
 		const TRange<FFrameNumber> PlaybackRange = MovieScene->GetPlaybackRange();
-		StartFrame = UE::MovieScene::DiscreteInclusiveLower(PlaybackRange);
-		DurationFrames = PlaybackRange.Size<FFrameNumber>();
+
+		const FFrameNumber StartTick = UE::MovieScene::DiscreteInclusiveLower(PlaybackRange);
+		const FFrameTime StartTime = ConvertFrameTime(StartTick, TickResolution, DisplayRate);
+
+		const FFrameNumber EndTick = UE::MovieScene::DiscreteExclusiveUpper(PlaybackRange);
+		const FFrameTime EndTime = ConvertFrameTime(EndTick, TickResolution, DisplayRate);
+
+		// Level sequences round starting time to a frame so let's do the same.
+		if (StartOffset > 0)
+		{
+			StartFrame = StartTime.FloorToFrame() + StartOffset;
+		}
+		else
+		{
+			StartFrame = StartTime.FloorToFrame();
+		}
+
+		DurationFrames = (EndTime - StartFrame);
 	}
 	else
 	{
 		StartFrame = FFrameNumber(0);
-		DurationFrames = 0;
+		DurationFrames = FFrameTime(0);
 	}
 
 	PlayPosition.Reset(StartFrame);
@@ -347,11 +369,7 @@ void UCameraAnimationSequencePlayer::Play(bool bLoop, bool bRandomStartTime)
 	// Move the playback position randomly in our playback range if we want a random start time.
 	if (bRandomStartTime)
 	{
-		UMovieScene* MovieScene = Sequence->GetMovieScene();
-
-		const TRange<FFrameNumber> PlaybackRange = MovieScene->GetPlaybackRange();
-		
-		const int32 RandomStartFrameOffset = FMath::RandHelper(DurationFrames.Value);
+		const int32 RandomStartFrameOffset = FMath::RandHelper(DurationFrames.FrameNumber.Value);
 		PlayPosition.Reset(StartFrame + RandomStartFrameOffset);
 	}
 
@@ -378,18 +396,18 @@ void UCameraAnimationSequencePlayer::Update(FFrameTime NewPosition)
 		// Arguably we could have some cumulative animation mode running on some properties but let's call
 		// this a limitation for now.
 		//
-		while (NewPosition.FrameNumber >= StartFrame + DurationFrames)
+		while (NewPosition >= StartFrame + DurationFrames)
 		{
-			NewPosition.FrameNumber -= DurationFrames;
+			NewPosition -= DurationFrames;
 			PlayPosition.Reset(StartFrame);
 		}
 	}
 	else
 	{
 		// If we are reaching the end, update the sequence at the end time and stop.
-		if (NewPosition.FrameNumber >= StartFrame + DurationFrames)
+		if (NewPosition >= StartFrame + DurationFrames)
 		{
-			NewPosition = FFrameTime(DurationFrames);
+			NewPosition = FFrameTime(StartFrame + DurationFrames);
 			bShouldStop = true;
 		}
 	}
@@ -444,4 +462,3 @@ void UCameraAnimationSequencePlayer::EndScrubbing()
 	Status = EMovieScenePlayerStatus::Playing;
 }
 
-
diff --git a/Engine/Plugins/MovieScene/TemplateSequence/Source/TemplateSequence/Private/SequenceCameraShake.cpp b/Engine/Plugins/MovieScene/TemplateSequence/Source/TemplateSequence/Private/SequenceCameraShake.cpp
--- a/Engine/Plugins/MovieScene/TemplateSequence/Source/TemplateSequence/Private/SequenceCameraShake.cpp
+++ b/Engine/Plugins/MovieScene/TemplateSequence/Source/TemplateSequence/Private/SequenceCameraShake.cpp
@@ -73,8 +73,10 @@ void USequenceCameraShakePattern::StartShakePatternImpl(const FCameraShakeStartP
 
 void USequenceCameraShakePattern::UpdateShakePatternImpl(const FCameraShakeUpdateParams& Params, FCameraShakeUpdateResult& OutResult)
 {
-	const FFrameRate InputRate = Player->GetInputRate();
-	const FFrameTime NewPosition = Player->GetCurrentPosition() + Params.DeltaTime * PlayRate * InputRate;
+	const FFrameRate DisplayRate = Player->GetInputRate();
+	const FFrameTime CurrentPosition = Player->GetCurrentPosition();
+
+	const FFrameTime NewPosition = CurrentPosition + Params.DeltaTime * PlayRate * DisplayRate;
 	UpdateCamera(NewPosition, Params.POV, OutResult);
 }
 
diff --git a/Engine/Plugins/MovieScene/TemplateSequence/Source/TemplateSequence/Public/CameraAnimationSequencePlayer.h b/Engine/Plugins/MovieScene/TemplateSequence/Source/TemplateSequence/Public/CameraAnimationSequencePlayer.h
--- a/Engine/Plugins/MovieScene/TemplateSequence/Source/TemplateSequence/Public/CameraAnimationSequencePlayer.h
+++ b/Engine/Plugins/MovieScene/TemplateSequence/Source/TemplateSequence/Public/CameraAnimationSequencePlayer.h
@@ -108,14 +108,39 @@ public:
 	UCameraAnimationSequencePlayer(const FObjectInitializer& ObjInit);
 	virtual ~UCameraAnimationSequencePlayer();
 
-	/** Initializes this player with the given sequence */
-	void Initialize(UMovieSceneSequence* InSequence);
-	/** Start playing the sequence */
+	/**
+	 * Initializes this player with the given sequence
+	 *
+	 * @param InSequence    The sequence to play
+	 * @param StartOffset   The offset to start at, in frames (display rate)
+	 */
+	void Initialize(UMovieSceneSequence* InSequence, int32 StartOffset = 0);
+
+	/**
+	 * Start playing the sequence
+	 *
+	 * @param bLoop              Whether to loop playback
+	 * @param bRandomStartTime   Whether to start at a random time inside the playback range
+	 * 
+	 * Note that if StartOffset was set, the random start time will be chosen within
+	 * the reduced (offset) playback range.
+	 */
 	void Play(bool bLoop = false, bool bRandomStartTime = false);
-	/** Advance play to the given time */
+
+	/**
+	 * Advance play to the given time
+	 *
+	 * @param NewPosition   The time to advance to, in ticks
+	 */
 	void Update(FFrameTime NewPosition);
-	/** Jumps to the given time */
+
+	/**
+	 * Jumps to the given time, in ticks
+	 *
+	 * @param New Position   The time to jump to, in ticks
+	 */
 	void Jump(FFrameTime NewPosition);
+
 	/** Stop playing the sequence */
 	void Stop();
 
@@ -126,14 +151,14 @@ public:
 
 	/** Gets whether playback is looping */
 	bool GetIsLooping() const { return bIsLooping; }
-	/** Gets the current play position */
-	FFrameTime GetCurrentPosition() const { return PlayPosition.GetCurrentPosition(); }
-	/** Get the sequence display resolution */
+	/** Get the sequence display rate */
 	FFrameRate GetInputRate() const { return PlayPosition.GetInputRate(); }
 	/** Get the sequence tick resolution */
 	FFrameRate GetOutputRate() const { return PlayPosition.GetOutputRate(); }
-	/** Get the duration of the current sequence */
-	FFrameNumber GetDuration() const;
+	/** Get the duration of the current sequence in frames (display rate) */
+	FFrameTime GetDuration() const;
+	/** Gets the current play position in frames (display rate) */
+	FFrameTime GetCurrentPosition() const;
 
 	/** Sets an object that can be used to bind everything in the sequence */
 	void SetBoundObjectOverride(UObject* InObject);
@@ -181,7 +206,7 @@ private:
 	FFrameNumber StartFrame;
 
 	/** The sequence duration in frames */
-	FFrameNumber DurationFrames;
+	FFrameTime DurationFrames;
 
 	/** Whether we should be looping */
 	bool bIsLooping;
