Sequencer: Clean up terminology

diff --git a/Engine/Plugins/Animation/LiveLink/Source/LiveLinkSequencer/Private/LiveLinkSequencerSettings.cpp b/Engine/Plugins/Animation/LiveLink/Source/LiveLinkSequencer/Private/LiveLinkSequencerSettings.cpp
--- a/Engine/Plugins/Animation/LiveLink/Source/LiveLinkSequencer/Private/LiveLinkSequencerSettings.cpp
+++ b/Engine/Plugins/Animation/LiveLink/Source/LiveLinkSequencer/Private/LiveLinkSequencerSettings.cpp
@@ -12,12 +12,12 @@ FName ULiveLinkSequencerSettings::GetCategoryName() const
 #if WITH_EDITOR
 FText ULiveLinkSequencerSettings::GetSectionText() const
 {
-	return NSLOCTEXT("LiveLinkSequencerSettings", "LiveLinkSequencerSettingsSection", "Live Link Sequencer");
+	return NSLOCTEXT("LiveLinkSequenceEditorSettings", "LiveLinkSequenceEditorSettingsSection", "Live Link Sequence Editor");
 }
 
 FName ULiveLinkSequencerSettings::GetSectionName() const
 {
-	return TEXT("Live Link Sequencer");
+	return TEXT("Live Link Sequence Editor");
 }
 
 #endif
diff --git a/Engine/Plugins/Animation/LiveLink/Source/LiveLinkSequencer/Public/LiveLinkSequencerSettings.h b/Engine/Plugins/Animation/LiveLink/Source/LiveLinkSequencer/Public/LiveLinkSequencerSettings.h
--- a/Engine/Plugins/Animation/LiveLink/Source/LiveLinkSequencer/Public/LiveLinkSequencerSettings.h
+++ b/Engine/Plugins/Animation/LiveLink/Source/LiveLinkSequencer/Public/LiveLinkSequencerSettings.h
@@ -13,7 +13,7 @@ class UMovieSceneLiveLinkControllerTrackRecorder;
 
 
 /**
- * Settings for LiveLink Sequencer
+ * Settings for LiveLink Sequence Editor
  */
 UCLASS(config=Game)
 class LIVELINKSEQUENCER_API ULiveLinkSequencerSettings : public UDeveloperSettings
diff --git a/Engine/Plugins/MovieScene/LevelSequenceEditor/Source/LevelSequenceEditor/Private/CinematicViewport/SCinematicLevelViewport.cpp b/Engine/Plugins/MovieScene/LevelSequenceEditor/Source/LevelSequenceEditor/Private/CinematicViewport/SCinematicLevelViewport.cpp
--- a/Engine/Plugins/MovieScene/LevelSequenceEditor/Source/LevelSequenceEditor/Private/CinematicViewport/SCinematicLevelViewport.cpp
+++ b/Engine/Plugins/MovieScene/LevelSequenceEditor/Source/LevelSequenceEditor/Private/CinematicViewport/SCinematicLevelViewport.cpp
@@ -416,7 +416,7 @@ void SCinematicLevelViewport::Construct(const FArguments& InArgs)
 					[
 						SNew(STextBlock)
 						.ColorAndOpacity(Gray)
-						.Text(LOCTEXT("NoSequencerMessage", "No active Level Sequencer detected. Please edit a Level Sequence to enable full controls."))
+						.Text(LOCTEXT("NoSequencerMessage", "No active Level Sequence Editor detected. Please edit a Level Sequence to enable full controls."))
 					]
 				]
 			]
diff --git a/Engine/Plugins/MovieScene/LevelSequenceEditor/Source/LevelSequenceEditor/Private/LevelSequenceEditorModule.cpp b/Engine/Plugins/MovieScene/LevelSequenceEditor/Source/LevelSequenceEditor/Private/LevelSequenceEditorModule.cpp
--- a/Engine/Plugins/MovieScene/LevelSequenceEditor/Source/LevelSequenceEditor/Private/LevelSequenceEditorModule.cpp
+++ b/Engine/Plugins/MovieScene/LevelSequenceEditor/Source/LevelSequenceEditor/Private/LevelSequenceEditorModule.cpp
@@ -167,9 +167,9 @@ protected:
 
 		if (SettingsModule != nullptr)
 		{
-			SettingsModule->RegisterSettings("Project", "Plugins", "LevelSequencer",
-				LOCTEXT("LevelSequencerSettingsName", "Level Sequencer"),
-				LOCTEXT("LevelSequencerSettingsDescription", "Configure the Level Sequence Editor."),
+			SettingsModule->RegisterSettings("Project", "Plugins", "LevelSequenceEditor",
+				LOCTEXT("LevelSequenceEditorProjectSettingsName", "Level Sequence Editor"),
+				LOCTEXT("LevelSequenceEditorProjectSettingsDescription", "Configure the Level Sequence Editor."),
 				GetMutableDefault<ULevelSequenceEditorSettings>()
 			);
 
@@ -250,7 +250,7 @@ protected:
 
 		if (SettingsModule != nullptr)
 		{
-			SettingsModule->UnregisterSettings("Project", "Plugins", "LevelSequencer");
+			SettingsModule->UnregisterSettings("Project", "Plugins", "LevelSequenceEditor");
 
 			SettingsModule->UnregisterSettings("Editor", "ContentEditors", "LevelSequenceEditor");
 		}
diff --git a/Engine/Plugins/MovieScene/LevelSequenceEditor/Source/LevelSequenceEditor/Public/LevelSequenceEditorSubsystem.h b/Engine/Plugins/MovieScene/LevelSequenceEditor/Source/LevelSequenceEditor/Public/LevelSequenceEditorSubsystem.h
--- a/Engine/Plugins/MovieScene/LevelSequenceEditor/Source/LevelSequenceEditor/Public/LevelSequenceEditorSubsystem.h
+++ b/Engine/Plugins/MovieScene/LevelSequenceEditor/Source/LevelSequenceEditor/Public/LevelSequenceEditorSubsystem.h
@@ -40,7 +40,7 @@ struct FMovieSceneScriptingParams
 
 /**
 * ULevelSequenceEditorSubsystem
-* Subsystem for level sequencer related utilities to scripts
+* Subsystem for level sequence editor related utilities to scripts
 */
 UCLASS(Blueprintable)
 class LEVELSEQUENCEEDITOR_API ULevelSequenceEditorSubsystem : public UEditorSubsystem
diff --git a/Engine/Plugins/MovieScene/TemplateSequence/Source/TemplateSequenceEditor/Private/TemplateSequenceEditorModule.cpp b/Engine/Plugins/MovieScene/TemplateSequence/Source/TemplateSequenceEditor/Private/TemplateSequenceEditorModule.cpp
--- a/Engine/Plugins/MovieScene/TemplateSequence/Source/TemplateSequenceEditor/Private/TemplateSequenceEditorModule.cpp
+++ b/Engine/Plugins/MovieScene/TemplateSequence/Source/TemplateSequenceEditor/Private/TemplateSequenceEditorModule.cpp
@@ -62,9 +62,9 @@ private:
 
 		if (SettingsModule != nullptr)
 		{
-			SettingsModule->RegisterSettings("Project", "Plugins", "TemplateSequencer",
-				LOCTEXT("TemplateSequencerSettingsName", "Template Sequencer"),
-				LOCTEXT("TemplateSequencerSettingsDescription", "Configure the Template Sequence Editor."),
+			SettingsModule->RegisterSettings("Project", "Plugins", "TemplateSequenceEditor",
+				LOCTEXT("TemplateSequenceEditorProjectSettingsName", "Template Sequence Editor"),
+				LOCTEXT("TemplateSequenceEditorProjectSettingsDescription", "Configure the Template Sequence Editor."),
 				GetMutableDefault<UTemplateSequenceEditorSettings>()
 			);
 
@@ -83,7 +83,7 @@ private:
 
 		if (SettingsModule != nullptr)
 		{
-			SettingsModule->UnregisterSettings("Project", "Plugins", "TemplateSequencer");
+			SettingsModule->UnregisterSettings("Project", "Plugins", "TemplateSequenceEditor");
 			SettingsModule->UnregisterSettings("Editor", "ContentEditors", "TemplateSequenceEditor");
 		}
 	}
diff --git a/Engine/Source/Editor/Sequencer/Private/SequencerEdMode.cpp b/Engine/Source/Editor/Sequencer/Private/SequencerEdMode.cpp
--- a/Engine/Source/Editor/Sequencer/Private/SequencerEdMode.cpp
+++ b/Engine/Source/Editor/Sequencer/Private/SequencerEdMode.cpp
@@ -52,7 +52,7 @@ const FEditorModeID FSequencerEdMode::EM_SequencerMode(TEXT("EM_SequencerMode"))
 static TAutoConsoleVariable<bool> CVarDrawMeshTrails(
 	TEXT("Sequencer.DrawMeshTrails"),
 	true,
-	TEXT("Toggle to show or hide Level Sequencer VR Editor trails"));
+	TEXT("Toggle to show or hide Level Sequence VR Editor trails"));
 
 //if true still use old motion trails for sequencer objects.
 TAutoConsoleVariable<bool> CVarUseOldSequencerMotionTrails(
