Cameras: fix incorrect blending out of sequencer camera shake patterns

The previous code was teleporting the sequence to its blend out time, but it should in fact simply keep playing and let itself be blended out by the state struct's time management. We just need to explicitly stop the player once that blend out is finished.

diff --git a/Engine/Plugins/MovieScene/TemplateSequence/Source/TemplateSequence/Private/CameraAnimationSequencePlayer.cpp b/Engine/Plugins/MovieScene/TemplateSequence/Source/TemplateSequence/Private/CameraAnimationSequencePlayer.cpp
--- a/Engine/Plugins/MovieScene/TemplateSequence/Source/TemplateSequence/Private/CameraAnimationSequencePlayer.cpp
+++ b/Engine/Plugins/MovieScene/TemplateSequence/Source/TemplateSequence/Private/CameraAnimationSequencePlayer.cpp
@@ -436,16 +436,19 @@ void UCameraAnimationSequencePlayer::Jump(FFrameTime NewPosition)
 
 void UCameraAnimationSequencePlayer::Stop()
 {
-	Status = EMovieScenePlayerStatus::Stopped;
+	if (Status != EMovieScenePlayerStatus::Stopped)
+	{
+		Status = EMovieScenePlayerStatus::Stopped;
 
-	PlayPosition.Reset(StartFrame);
+		PlayPosition.Reset(StartFrame);
 
-	if (TSharedPtr<FMovieSceneEntitySystemRunner> Runner = RootTemplateInstance.GetRunner())
-	{
-		if (Runner->QueueFinalUpdate(RootTemplateInstance.GetRootInstanceHandle()))
+		if (TSharedPtr<FMovieSceneEntitySystemRunner> Runner = RootTemplateInstance.GetRunner())
 		{
-			// @todo: we should really only call flush _once_ for all camera anim sequences
-			Runner->Flush();
+			if (Runner->QueueFinalUpdate(RootTemplateInstance.GetRootInstanceHandle()))
+			{
+				// @todo: we should really only call flush _once_ for all camera anim sequences
+				Runner->Flush();
+			}
 		}
 	}
 }
diff --git a/Engine/Plugins/MovieScene/TemplateSequence/Source/TemplateSequence/Private/SequenceCameraShake.cpp b/Engine/Plugins/MovieScene/TemplateSequence/Source/TemplateSequence/Private/SequenceCameraShake.cpp
--- a/Engine/Plugins/MovieScene/TemplateSequence/Source/TemplateSequence/Private/SequenceCameraShake.cpp
+++ b/Engine/Plugins/MovieScene/TemplateSequence/Source/TemplateSequence/Private/SequenceCameraShake.cpp
@@ -77,26 +77,46 @@ void USequenceCameraShakePattern::StartShakePatternImpl(const FCameraShakeStartP
 
 void USequenceCameraShakePattern::UpdateShakePatternImpl(const FCameraShakeUpdateParams& Params, FCameraShakeUpdateResult& OutResult)
 {
-	const FFrameRate DisplayRate = Player->GetInputRate();
-	const FFrameTime CurrentPosition = Player->GetCurrentPosition();
+	const float BlendWeight = State.Update(Params.DeltaTime);
+	if (State.IsPlaying())
+	{
+		const FFrameRate DisplayRate = Player->GetInputRate();
+		const FFrameTime CurrentPosition = Player->GetCurrentPosition();
 
-	const FFrameTime NewPosition = CurrentPosition + Params.DeltaTime * PlayRate * DisplayRate;
-	UpdateCamera(NewPosition, Params.POV, OutResult);
+		const FFrameTime NewPosition = CurrentPosition + Params.DeltaTime * PlayRate * DisplayRate;
+		UpdateCamera(NewPosition, Params.POV, OutResult);
 
-	const float BlendWeight = State.Update(Params.DeltaTime);
-	OutResult.ApplyScale(BlendWeight);
+		OutResult.ApplyScale(BlendWeight);
+	}
+	else
+	{
+		// We could have stopped playing if we were blending out as a result of a call to Stop(false)
+		// In this case, the sequence might still be playing (i.e. in a valid position inside its playback
+		// range) and we need to explicitly stop it.
+		check(Player);
+		Player->Stop();
+	}
 }
 
 void USequenceCameraShakePattern::ScrubShakePatternImpl(const FCameraShakeScrubParams& Params, FCameraShakeUpdateResult& OutResult)
 {
 	Player->StartScrubbing();
 	{
-		const FFrameRate InputRate = Player->GetInputRate();
-		const FFrameTime NewPosition = Params.AbsoluteTime * PlayRate * InputRate;
-		UpdateCamera(NewPosition, Params.POV, OutResult);
-		
 		const float BlendWeight = State.Scrub(Params.AbsoluteTime);
-		OutResult.ApplyScale(BlendWeight);
+		if (State.IsPlaying())
+		{
+			const FFrameRate InputRate = Player->GetInputRate();
+			const FFrameTime NewPosition = Params.AbsoluteTime * PlayRate * InputRate;
+			UpdateCamera(NewPosition, Params.POV, OutResult);
+
+			OutResult.ApplyScale(BlendWeight);
+		}
+		else
+		{
+			// See the similar else clause in UpdateShakePatternImpl. 
+			check(Player);
+			Player->Stop();
+		}
 	}
 	Player->EndScrubbing();
 }
@@ -118,20 +138,8 @@ void USequenceCameraShakePattern::StopShakePatternImpl(const FCameraShakeStopPar
 			// Stop playing!
 			Player->Stop();
 		}
-		else
-		{
-			// Move the playback position to the start of the blend out.
-			const TRange<FFrameNumber> PlaybackRange = MovieScene->GetPlaybackRange();
-			if (PlaybackRange.HasUpperBound())
-			{
-				const FFrameRate InputRate = Player->GetInputRate();
-				const FFrameTime BlendOutTimeInFrames  = BlendOutTime * InputRate;
-				const FFrameTime BlendOutStartFrame = FFrameTime(PlaybackRange.GetUpperBoundValue()) - BlendOutTimeInFrames;
-				Player->Jump(BlendOutStartFrame);
-			}
-		}
 
-		// Stop tracking time.
+		// Make our state tracking stop or start blending out.
 		State.Stop(Params.bImmediately);
 	}
 }
