Sequencer: Allow track editors if it's supported by the sequence (and not just by level sequence)

Fix regression caused by track permissions 20284129

diff --git a/Engine/Plugins/Animation/LiveLink/Source/LiveLinkSequencer/Private/Sequencer/LiveLinkPropertyTrackEditor.cpp b/Engine/Plugins/Animation/LiveLink/Source/LiveLinkSequencer/Private/Sequencer/LiveLinkPropertyTrackEditor.cpp
--- a/Engine/Plugins/Animation/LiveLink/Source/LiveLinkSequencer/Private/Sequencer/LiveLinkPropertyTrackEditor.cpp
+++ b/Engine/Plugins/Animation/LiveLink/Source/LiveLinkSequencer/Private/Sequencer/LiveLinkPropertyTrackEditor.cpp
@@ -206,12 +206,14 @@ TSharedRef<ISequencerSection> FLiveLinkPropertyTrackEditor::MakeSectionInterface
 
 bool FLiveLinkPropertyTrackEditor::SupportsSequence(UMovieSceneSequence* InSequence) const
 {
-	if (InSequence && InSequence->IsTrackSupported(UMovieSceneLiveLinkTrack::StaticClass()) == ETrackSupport::NotSupported)
+	ETrackSupport TrackSupported = InSequence ? InSequence->IsTrackSupported(UMovieSceneLiveLinkTrack::StaticClass()) : ETrackSupport::Default;
+
+	if (TrackSupported == ETrackSupport::NotSupported)
 	{
 		return false;
 	}
 
-	return InSequence && InSequence->IsA(ULevelSequence::StaticClass());
+	return (InSequence && InSequence->IsA(ULevelSequence::StaticClass())) || TrackSupported == ETrackSupport::Supported;
 }
 
 bool FLiveLinkPropertyTrackEditor::SupportsType(TSubclassOf<UMovieSceneTrack> Type) const
diff --git a/Engine/Plugins/Experimental/GeometryCollectionPlugin/Source/GeometryCollectionSequencer/Private/GeometryCollectionTrackEditor.cpp b/Engine/Plugins/Experimental/GeometryCollectionPlugin/Source/GeometryCollectionSequencer/Private/GeometryCollectionTrackEditor.cpp
--- a/Engine/Plugins/Experimental/GeometryCollectionPlugin/Source/GeometryCollectionSequencer/Private/GeometryCollectionTrackEditor.cpp
+++ b/Engine/Plugins/Experimental/GeometryCollectionPlugin/Source/GeometryCollectionSequencer/Private/GeometryCollectionTrackEditor.cpp
@@ -277,12 +277,14 @@ TSharedRef<ISequencerTrackEditor> FGeometryCollectionTrackEditor::CreateTrackEdi
 
 bool FGeometryCollectionTrackEditor::SupportsSequence(UMovieSceneSequence* InSequence) const
 {
-	if (InSequence && InSequence->IsTrackSupported(UMovieSceneGeometryCollectionTrack::StaticClass()) == ETrackSupport::NotSupported)
+	ETrackSupport TrackSupported = InSequence ? InSequence->IsTrackSupported(UMovieSceneGeometryCollectionTrack::StaticClass()) : ETrackSupport::Default;
+
+	if (TrackSupported == ETrackSupport::NotSupported)
 	{
 		return false;
 	}
 
-	return InSequence && InSequence->IsA(ULevelSequence::StaticClass());
+	return (InSequence && InSequence->IsA(ULevelSequence::StaticClass())) || TrackSupported == ETrackSupport::Supported;
 }
 
 bool FGeometryCollectionTrackEditor::SupportsType( TSubclassOf<UMovieSceneTrack> Type ) const
diff --git a/Engine/Plugins/Media/MediaCompositing/Source/MediaCompositingEditor/Private/Sequencer/MediaPlayerPropertyTrackEditor.cpp b/Engine/Plugins/Media/MediaCompositing/Source/MediaCompositingEditor/Private/Sequencer/MediaPlayerPropertyTrackEditor.cpp
--- a/Engine/Plugins/Media/MediaCompositing/Source/MediaCompositingEditor/Private/Sequencer/MediaPlayerPropertyTrackEditor.cpp
+++ b/Engine/Plugins/Media/MediaCompositing/Source/MediaCompositingEditor/Private/Sequencer/MediaPlayerPropertyTrackEditor.cpp
@@ -191,12 +191,14 @@ TSharedRef<ISequencerSection> FMediaPlayerPropertyTrackEditor::MakeSectionInterf
 
 bool FMediaPlayerPropertyTrackEditor::SupportsSequence(UMovieSceneSequence* InSequence) const
 {
-	if (InSequence && InSequence->IsTrackSupported(UMovieSceneMediaPlayerPropertyTrack::StaticClass()) == ETrackSupport::NotSupported)
+	ETrackSupport TrackSupported = InSequence ? InSequence->IsTrackSupported(UMovieSceneMediaPlayerPropertyTrack::StaticClass()) : ETrackSupport::Default;
+
+	if (TrackSupported == ETrackSupport::NotSupported)
 	{
 		return false;
 	}
 
-	return InSequence && InSequence->IsA(ULevelSequence::StaticClass());
+	return (InSequence && InSequence->IsA(ULevelSequence::StaticClass())) || TrackSupported == ETrackSupport::Supported;
 }
 
 const FSlateBrush* FMediaPlayerPropertyTrackEditor::GetIconBrush() const
diff --git a/Engine/Plugins/Media/MediaCompositing/Source/MediaCompositingEditor/Private/Sequencer/MediaTrackEditor.cpp b/Engine/Plugins/Media/MediaCompositing/Source/MediaCompositingEditor/Private/Sequencer/MediaTrackEditor.cpp
--- a/Engine/Plugins/Media/MediaCompositing/Source/MediaCompositingEditor/Private/Sequencer/MediaTrackEditor.cpp
+++ b/Engine/Plugins/Media/MediaCompositing/Source/MediaCompositingEditor/Private/Sequencer/MediaTrackEditor.cpp
@@ -186,12 +186,14 @@ TSharedRef<ISequencerSection> FMediaTrackEditor::MakeSectionInterface(UMovieScen
 
 bool FMediaTrackEditor::SupportsSequence(UMovieSceneSequence* InSequence) const 
 {
-	if (InSequence && InSequence->IsTrackSupported(UMovieSceneMediaTrack::StaticClass()) == ETrackSupport::NotSupported)
+	ETrackSupport TrackSupported = InSequence ? InSequence->IsTrackSupported(UMovieSceneMediaTrack::StaticClass()) : ETrackSupport::Default;
+
+	if (TrackSupported == ETrackSupport::NotSupported)
 	{
 		return false;
 	}
 
-	return InSequence && InSequence->IsA(ULevelSequence::StaticClass());
+	return (InSequence && InSequence->IsA(ULevelSequence::StaticClass())) || TrackSupported == ETrackSupport::Supported;
 }
 
 bool FMediaTrackEditor::SupportsType(TSubclassOf<UMovieSceneTrack> TrackClass) const
diff --git a/Engine/Plugins/MovieScene/TemplateSequence/Source/TemplateSequenceEditor/Private/TrackEditors/TemplateSequenceTrackEditor.cpp b/Engine/Plugins/MovieScene/TemplateSequence/Source/TemplateSequenceEditor/Private/TrackEditors/TemplateSequenceTrackEditor.cpp
--- a/Engine/Plugins/MovieScene/TemplateSequence/Source/TemplateSequenceEditor/Private/TrackEditors/TemplateSequenceTrackEditor.cpp
+++ b/Engine/Plugins/MovieScene/TemplateSequence/Source/TemplateSequenceEditor/Private/TrackEditors/TemplateSequenceTrackEditor.cpp
@@ -44,12 +44,14 @@ bool FTemplateSequenceTrackEditor::SupportsType(TSubclassOf<UMovieSceneTrack> Ty
 
 bool FTemplateSequenceTrackEditor::SupportsSequence(UMovieSceneSequence* InSequence) const
 {
-	if (InSequence && InSequence->IsTrackSupported(UTemplateSequenceTrack::StaticClass()) == ETrackSupport::NotSupported)
+	ETrackSupport TrackSupported = InSequence ? InSequence->IsTrackSupported(UTemplateSequenceTrack::StaticClass()) : ETrackSupport::Default;
+
+	if (TrackSupported == ETrackSupport::NotSupported)
 	{
 		return false;
 	}
 
-	return InSequence && (InSequence->IsA(ULevelSequence::StaticClass()) || InSequence->IsA(UTemplateSequence::StaticClass()));
+	return (InSequence && (InSequence->IsA(ULevelSequence::StaticClass()) || InSequence->IsA(UTemplateSequence::StaticClass()))) || TrackSupported == ETrackSupport::Supported;
 }
 
 void FTemplateSequenceTrackEditor::BuildObjectBindingTrackMenu(FMenuBuilder& MenuBuilder, const TArray<FGuid>& ObjectBindings, const UClass* ObjectClass)
diff --git a/Engine/Plugins/Runtime/GameplayAbilities/Source/GameplayAbilitiesEditor/GameplayAbilitiesEditor.Build.cs b/Engine/Plugins/Runtime/GameplayAbilities/Source/GameplayAbilitiesEditor/GameplayAbilitiesEditor.Build.cs
--- a/Engine/Plugins/Runtime/GameplayAbilities/Source/GameplayAbilitiesEditor/GameplayAbilitiesEditor.Build.cs
+++ b/Engine/Plugins/Runtime/GameplayAbilities/Source/GameplayAbilitiesEditor/GameplayAbilitiesEditor.Build.cs
@@ -29,6 +29,7 @@ namespace UnrealBuildTool.Rules
 					"Kismet",
 					"KismetCompiler",
 					"GraphEditor",
+					"LevelSequence",
 					"MainFrame",
 					"EditorFramework",
 					"UnrealEd",
diff --git a/Engine/Plugins/Runtime/GameplayAbilities/Source/GameplayAbilitiesEditor/Private/Sequencer/GameplayCueTrackEditor.cpp b/Engine/Plugins/Runtime/GameplayAbilities/Source/GameplayAbilitiesEditor/Private/Sequencer/GameplayCueTrackEditor.cpp
--- a/Engine/Plugins/Runtime/GameplayAbilities/Source/GameplayAbilitiesEditor/Private/Sequencer/GameplayCueTrackEditor.cpp
+++ b/Engine/Plugins/Runtime/GameplayAbilities/Source/GameplayAbilitiesEditor/Private/Sequencer/GameplayCueTrackEditor.cpp
@@ -12,6 +12,7 @@
 #include "Widgets/SBoxPanel.h"
 #include "SequencerUtilities.h"
 #include "MovieSceneSequenceEditor.h"
+#include "LevelSequence.h"
 
 #define LOCTEXT_NAMESPACE "FGameplayCueTrackEditor"
 
@@ -140,12 +141,14 @@ bool FGameplayCueTrackEditor::SupportsType(TSubclassOf<UMovieSceneTrack> Type) c
 
 bool FGameplayCueTrackEditor::SupportsSequence(UMovieSceneSequence* InSequence) const
 {
-	if (InSequence && InSequence->IsTrackSupported(UMovieSceneGameplayCueTrack::StaticClass()) == ETrackSupport::NotSupported)
+	ETrackSupport TrackSupported = InSequence ? InSequence->IsTrackSupported(UMovieSceneGameplayCueTrack::StaticClass()) : ETrackSupport::Default;
+
+	if (TrackSupported == ETrackSupport::NotSupported)
 	{
 		return false;
 	}
 
-	return true;
+	return (InSequence && InSequence->IsA(ULevelSequence::StaticClass())) || TrackSupported == ETrackSupport::Supported;
 }
 
 const FSlateBrush* FGameplayCueTrackEditor::GetIconBrush() const
diff --git a/Engine/Plugins/Runtime/GeometryCache/Source/GeometryCacheSequencer/Classes/GeometryCacheTrackEditor.cpp b/Engine/Plugins/Runtime/GeometryCache/Source/GeometryCacheSequencer/Classes/GeometryCacheTrackEditor.cpp
--- a/Engine/Plugins/Runtime/GeometryCache/Source/GeometryCacheSequencer/Classes/GeometryCacheTrackEditor.cpp
+++ b/Engine/Plugins/Runtime/GeometryCache/Source/GeometryCacheSequencer/Classes/GeometryCacheTrackEditor.cpp
@@ -276,12 +276,14 @@ TSharedRef<ISequencerTrackEditor> FGeometryCacheTrackEditor::CreateTrackEditor(T
 
 bool FGeometryCacheTrackEditor::SupportsSequence(UMovieSceneSequence* InSequence) const
 {
-	if (InSequence && InSequence->IsTrackSupported(UMovieSceneGeometryCacheTrack::StaticClass()) == ETrackSupport::NotSupported)
+	ETrackSupport TrackSupported = InSequence ? InSequence->IsTrackSupported(UMovieSceneGeometryCacheTrack::StaticClass()) : ETrackSupport::Default;
+
+	if (TrackSupported == ETrackSupport::NotSupported)
 	{
 		return false;
 	}
 
-	return InSequence && InSequence->IsA(ULevelSequence::StaticClass());
+	return (InSequence && InSequence->IsA(ULevelSequence::StaticClass())) || TrackSupported == ETrackSupport::Supported;
 }
 
 bool FGeometryCacheTrackEditor::SupportsType(TSubclassOf<UMovieSceneTrack> Type) const
diff --git a/Engine/Plugins/Runtime/HairStrands/Source/HairStrandsEditor/Private/GroomCacheTrackEditor.cpp b/Engine/Plugins/Runtime/HairStrands/Source/HairStrandsEditor/Private/GroomCacheTrackEditor.cpp
--- a/Engine/Plugins/Runtime/HairStrands/Source/HairStrandsEditor/Private/GroomCacheTrackEditor.cpp
+++ b/Engine/Plugins/Runtime/HairStrands/Source/HairStrandsEditor/Private/GroomCacheTrackEditor.cpp
@@ -270,12 +270,14 @@ TSharedRef<ISequencerTrackEditor> FGroomCacheTrackEditor::CreateTrackEditor(TSha
 
 bool FGroomCacheTrackEditor::SupportsSequence(UMovieSceneSequence* InSequence) const
 {
-	if (InSequence && InSequence->IsTrackSupported(UMovieSceneGroomCacheTrack::StaticClass()) == ETrackSupport::NotSupported)
+	ETrackSupport TrackSupported = InSequence ? InSequence->IsTrackSupported(UMovieSceneGroomCacheTrack::StaticClass()) : ETrackSupport::Default;
+
+	if (TrackSupported == ETrackSupport::NotSupported)
 	{
 		return false;
 	}
 
-	return InSequence && InSequence->IsA(ULevelSequence::StaticClass());
+	return (InSequence && InSequence->IsA(ULevelSequence::StaticClass())) || TrackSupported == ETrackSupport::Supported;
 }
 
 bool FGroomCacheTrackEditor::SupportsType(TSubclassOf<UMovieSceneTrack> Type) const
diff --git a/Engine/Plugins/VirtualProduction/LiveLinkCamera/Source/LiveLinkCameraRecording/Private/LiveLinkCameraControllerTrackEditor.cpp b/Engine/Plugins/VirtualProduction/LiveLinkCamera/Source/LiveLinkCameraRecording/Private/LiveLinkCameraControllerTrackEditor.cpp
--- a/Engine/Plugins/VirtualProduction/LiveLinkCamera/Source/LiveLinkCameraRecording/Private/LiveLinkCameraControllerTrackEditor.cpp
+++ b/Engine/Plugins/VirtualProduction/LiveLinkCamera/Source/LiveLinkCameraRecording/Private/LiveLinkCameraControllerTrackEditor.cpp
@@ -12,12 +12,14 @@ TSharedRef<ISequencerTrackEditor> FLiveLinkCameraControllerTrackEditor::CreateTr
 
 bool FLiveLinkCameraControllerTrackEditor::SupportsSequence(UMovieSceneSequence* InSequence) const
 {
-	if (InSequence && InSequence->IsTrackSupported(UMovieSceneLiveLinkCameraControllerTrack::StaticClass()) == ETrackSupport::NotSupported)
+	ETrackSupport TrackSupported = InSequence ? InSequence->IsTrackSupported(UMovieSceneLiveLinkCameraControllerTrack::StaticClass()) : ETrackSupport::Default;
+
+	if (TrackSupported == ETrackSupport::NotSupported)
 	{
 		return false;
 	}
 
-	return InSequence && InSequence->IsA(ULevelSequence::StaticClass());
+	return (InSequence && InSequence->IsA(ULevelSequence::StaticClass())) || TrackSupported == ETrackSupport::Supported;
 }
 
 bool FLiveLinkCameraControllerTrackEditor::SupportsType(TSubclassOf<UMovieSceneTrack> Type) const
