Sequencer: Remove unused StartScrubbing/EndScrubbing which was causing an ensure to fire in EndScrubbing when stop was called before EndScrubbing

diff --git a/Engine/Plugins/MovieScene/TemplateSequence/Source/TemplateSequence/Private/CameraAnimationSequencePlayer.cpp b/Engine/Plugins/MovieScene/TemplateSequence/Source/TemplateSequence/Private/CameraAnimationSequencePlayer.cpp
--- a/Engine/Plugins/MovieScene/TemplateSequence/Source/TemplateSequence/Private/CameraAnimationSequencePlayer.cpp
+++ b/Engine/Plugins/MovieScene/TemplateSequence/Source/TemplateSequence/Private/CameraAnimationSequencePlayer.cpp
@@ -383,7 +383,7 @@ void UCameraAnimationSequencePlayer::Play(bool bLoop, bool bRandomStartTime)
 
 void UCameraAnimationSequencePlayer::Update(FFrameTime NewPosition)
 {
-	check(Status == EMovieScenePlayerStatus::Playing || Status == EMovieScenePlayerStatus::Scrubbing);
+	check(Status == EMovieScenePlayerStatus::Playing);
 	check(RootTemplateInstance.IsValid());
 
 	bool bShouldStop = false;
diff --git a/Engine/Plugins/MovieScene/TemplateSequence/Source/TemplateSequence/Private/SequenceCameraShake.cpp b/Engine/Plugins/MovieScene/TemplateSequence/Source/TemplateSequence/Private/SequenceCameraShake.cpp
--- a/Engine/Plugins/MovieScene/TemplateSequence/Source/TemplateSequence/Private/SequenceCameraShake.cpp
+++ b/Engine/Plugins/MovieScene/TemplateSequence/Source/TemplateSequence/Private/SequenceCameraShake.cpp
@@ -100,25 +100,21 @@ void USequenceCameraShakePattern::UpdateShakePatternImpl(const FCameraShakeUpdat
 
 void USequenceCameraShakePattern::ScrubShakePatternImpl(const FCameraShakeScrubParams& Params, FCameraShakeUpdateResult& OutResult)
 {
-	Player->StartScrubbing();
+	const float BlendWeight = State.Scrub(Params.AbsoluteTime);
+	if (State.IsPlaying())
 	{
-		const float BlendWeight = State.Scrub(Params.AbsoluteTime);
-		if (State.IsPlaying())
-		{
-			const FFrameRate InputRate = Player->GetInputRate();
-			const FFrameTime NewPosition = Params.AbsoluteTime * PlayRate * InputRate;
-			UpdateCamera(NewPosition, Params.POV, OutResult);
+		const FFrameRate InputRate = Player->GetInputRate();
+		const FFrameTime NewPosition = Params.AbsoluteTime * PlayRate * InputRate;
+		UpdateCamera(NewPosition, Params.POV, OutResult);
 
-			OutResult.ApplyScale(BlendWeight);
-		}
-		else
-		{
-			// See the similar else clause in UpdateShakePatternImpl. 
-			check(Player);
-			Player->Stop();
-		}
+		OutResult.ApplyScale(BlendWeight);
+	}
+	else
+	{
+		// See the similar else clause in UpdateShakePatternImpl. 
+		check(Player);
+		Player->Stop();
 	}
-	Player->EndScrubbing();
 }
 
 bool USequenceCameraShakePattern::IsFinishedImpl() const
diff --git a/Engine/Plugins/MovieScene/TemplateSequence/Source/TemplateSequence/Public/CameraAnimationSequencePlayer.h b/Engine/Plugins/MovieScene/TemplateSequence/Source/TemplateSequence/Public/CameraAnimationSequencePlayer.h
--- a/Engine/Plugins/MovieScene/TemplateSequence/Source/TemplateSequence/Public/CameraAnimationSequencePlayer.h
+++ b/Engine/Plugins/MovieScene/TemplateSequence/Source/TemplateSequence/Public/CameraAnimationSequencePlayer.h
@@ -145,8 +145,11 @@ public:
 	void Stop();
 
 	/** Sets the player in scrub mode */
+	UE_DEPRECATED(5.3, "StartScrubbing has been deprecated as it has no functionality")
 	void StartScrubbing();
+
 	/** Ends scrub mode */
+	UE_DEPRECATED(5.3, "EndScrubbing has been deprecated as it has no functionality")
 	void EndScrubbing();
 
 	/** Gets whether playback is looping */
