Sequencer: register/unregister the camera animation stand-in object from its own module

This is a speculative fix for a crash that happens on shutdown, when it looks like the camera stand-in UClass does not exist anymore. This is possibly due to how we were unregistering it from the sequencer ECS custom accessors from a different module, and so maybe that module was getting unloaded after the one in which the camera stand-in lives. Either way, keeping things in the same module is probably better anyway.

diff --git a/Engine/Plugins/Cameras/GameplayCameras/Source/GameplayCameras/Private/GameplayCamerasModule.cpp b/Engine/Plugins/Cameras/GameplayCameras/Source/GameplayCameras/Private/GameplayCamerasModule.cpp
--- a/Engine/Plugins/Cameras/GameplayCameras/Source/GameplayCameras/Private/GameplayCamerasModule.cpp
+++ b/Engine/Plugins/Cameras/GameplayCameras/Source/GameplayCameras/Private/GameplayCamerasModule.cpp
@@ -24,8 +24,6 @@ public:
 		{
 			IModularFeatures::Get().RegisterModularFeature(ICameraModularFeature::GetModularFeatureName(), CameraModularFeature.Get());
 		}
-
-		UCameraAnimationSequenceCameraStandIn::RegisterCameraStandIn();
 	}
 
 	virtual void ShutdownModule() override
@@ -35,8 +33,6 @@ public:
 			IModularFeatures::Get().UnregisterModularFeature(ICameraModularFeature::GetModularFeatureName(), CameraModularFeature.Get());
 			CameraModularFeature = nullptr;
 		}
-
-		UCameraAnimationSequenceCameraStandIn::UnregisterCameraStandIn();
 	}
 
 private:
diff --git a/Engine/Plugins/MovieScene/TemplateSequence/Source/TemplateSequence/Private/TemplateSequence.cpp b/Engine/Plugins/MovieScene/TemplateSequence/Source/TemplateSequence/Private/TemplateSequence.cpp
--- a/Engine/Plugins/MovieScene/TemplateSequence/Source/TemplateSequence/Private/TemplateSequence.cpp
+++ b/Engine/Plugins/MovieScene/TemplateSequence/Source/TemplateSequence/Private/TemplateSequence.cpp
@@ -4,7 +4,6 @@
 #include "AssetRegistry/AssetData.h"
 #include "Engine/World.h"
 #include "Misc/PackageName.h"
-#include "Modules/ModuleManager.h"
 #include "MovieScene.h"
 #include "MovieSceneCommonHelpers.h"
 #include "MovieScenePossessable.h"
@@ -18,8 +17,6 @@
 
 #include UE_INLINE_GENERATED_CPP_BY_NAME(TemplateSequence)
 
-IMPLEMENT_MODULE(FDefaultModuleImpl, TemplateSequence);
-
 DEFINE_LOG_CATEGORY(LogTemplateSequence);
 
 UTemplateSequence::UTemplateSequence(const FObjectInitializer& ObjectInitializer)
diff --git a/Engine/Plugins/MovieScene/TemplateSequence/Source/TemplateSequence/Private/TemplateSequenceModule.cpp b/Engine/Plugins/MovieScene/TemplateSequence/Source/TemplateSequence/Private/TemplateSequenceModule.cpp
new file mode 100644
--- /dev/null
+++ b/Engine/Plugins/MovieScene/TemplateSequence/Source/TemplateSequence/Private/TemplateSequenceModule.cpp
@@ -0,0 +1,21 @@
+// Copyright Epic Games, Inc. All Rights Reserved.
+
+#include "Modules/ModuleManager.h"
+#include "CameraAnimationSequencePlayer.h"
+
+class FTemplateSequenceModule : public IModuleInterface
+{
+public:
+
+	virtual void StartupModule() override
+	{
+		UCameraAnimationSequenceCameraStandIn::RegisterCameraStandIn();
+	}
+
+	virtual void ShutdownModule() override
+	{
+		UCameraAnimationSequenceCameraStandIn::UnregisterCameraStandIn();
+	}
+};
+
+IMPLEMENT_MODULE(FTemplateSequenceModule, TemplateSequence);
diff --git a/Engine/Plugins/MovieScene/TemplateSequence/Source/TemplateSequence/Public/CameraAnimationSequencePlayer.h b/Engine/Plugins/MovieScene/TemplateSequence/Source/TemplateSequence/Public/CameraAnimationSequencePlayer.h
--- a/Engine/Plugins/MovieScene/TemplateSequence/Source/TemplateSequence/Public/CameraAnimationSequencePlayer.h
+++ b/Engine/Plugins/MovieScene/TemplateSequence/Source/TemplateSequence/Public/CameraAnimationSequencePlayer.h
@@ -75,16 +75,15 @@ public:
 	/** Recompute camera and lens settings after each frame */
 	void RecalcDerivedData();
 
+private:
+	void ResetDefaultValues(const FMinimalViewInfo& ViewInfo);
+	void UpdateInitialPropertyValues(UMovieSceneEntitySystemLinker* Linker);
+
 	/** Register the stand-in class with the sequencer ECS component registry */
 	static void RegisterCameraStandIn();
-
 	/** Unregister the stand-in class with the sequencer ECS component registry */
 	static void UnregisterCameraStandIn();
 
-private:
-	void ResetDefaultValues(const FMinimalViewInfo& ViewInfo);
-	void UpdateInitialPropertyValues(UMovieSceneEntitySystemLinker* Linker);
-
 private:
 	static bool bRegistered;
 
@@ -92,6 +91,8 @@ private:
 
 	bool bIsCineCamera = false;
 	float WorldToMeters = 0.f;
+
+	friend class FTemplateSequenceModule;
 };
 
 /**
