Template Sequence: Fix crash with ISequencer becoming invalid when opening a level sequence after a camera animation sequence

diff --git a/Engine/Plugins/MovieScene/TemplateSequence/Source/TemplateSequenceEditor/Private/Customizations/TemplateSequenceCustomizationBase.cpp b/Engine/Plugins/MovieScene/TemplateSequence/Source/TemplateSequenceEditor/Private/Customizations/TemplateSequenceCustomizationBase.cpp
--- a/Engine/Plugins/MovieScene/TemplateSequence/Source/TemplateSequenceEditor/Private/Customizations/TemplateSequenceCustomizationBase.cpp
+++ b/Engine/Plugins/MovieScene/TemplateSequence/Source/TemplateSequenceEditor/Private/Customizations/TemplateSequenceCustomizationBase.cpp
@@ -15,6 +15,7 @@ void FTemplateSequenceCustomizationBase::RegisterSequencerCustomization(FSequenc
 	Sequencer = &Builder.GetSequencer();
 	TemplateSequence = Cast<UTemplateSequence>(&Builder.GetFocusedSequence());
 
+	Sequencer->OnCloseEvent().AddRaw(this, &FTemplateSequenceCustomizationBase::OnSequencerClosed);
 	Sequencer->OnMovieSceneDataChanged().AddRaw(this, &FTemplateSequenceCustomizationBase::OnMovieSceneDataChanged);
 
 	FSequencerCustomizationInfo BaseInfo;
@@ -24,7 +25,11 @@ void FTemplateSequenceCustomizationBase::RegisterSequencerCustomization(FSequenc
 
 void FTemplateSequenceCustomizationBase::UnregisterSequencerCustomization()
 {
-	Sequencer->OnMovieSceneDataChanged().RemoveAll(this);
+	if (Sequencer)
+	{
+		Sequencer->OnCloseEvent().RemoveAll(this);
+		Sequencer->OnMovieSceneDataChanged().RemoveAll(this);
+	}
 
 	Sequencer = nullptr;
 	TemplateSequence = nullptr;
@@ -96,5 +101,16 @@ void FTemplateSequenceCustomizationBase::OnMovieSceneDataChanged(EMovieSceneData
 	}
 }
 
+void FTemplateSequenceCustomizationBase::OnSequencerClosed(TSharedRef<ISequencer> InSequencer)
+{
+	if (&InSequencer.Get() == Sequencer)
+	{
+		Sequencer->OnCloseEvent().RemoveAll(this);
+		Sequencer->OnMovieSceneDataChanged().RemoveAll(this);
+
+		Sequencer = nullptr;
+	}
+}
+
 #undef LOCTEXT_NAMESPACE
 
diff --git a/Engine/Plugins/MovieScene/TemplateSequence/Source/TemplateSequenceEditor/Private/Customizations/TemplateSequenceCustomizationBase.h b/Engine/Plugins/MovieScene/TemplateSequence/Source/TemplateSequenceEditor/Private/Customizations/TemplateSequenceCustomizationBase.h
--- a/Engine/Plugins/MovieScene/TemplateSequence/Source/TemplateSequenceEditor/Private/Customizations/TemplateSequenceCustomizationBase.h
+++ b/Engine/Plugins/MovieScene/TemplateSequence/Source/TemplateSequenceEditor/Private/Customizations/TemplateSequenceCustomizationBase.h
@@ -29,6 +29,7 @@ protected:
 private:
 	ESequencerPasteSupport OnPaste();
 	void OnMovieSceneDataChanged(EMovieSceneDataChangeType ChangeType);
+	void OnSequencerClosed(TSharedRef<ISequencer> InSequencer);
 
 	ISequencer* Sequencer;
 	UTemplateSequence* TemplateSequence;
