Sequencer: Defer deletion of sub sequence instances to the active runner

When a sequence is re-compiled mid-evaluation, we have to re-update the sequence updater to ensure that the set of imported entities is accurate. Previously, references to expiring sequence instances would be kept in the updater until the _next_ frame, when they would be removed from the map. Unfortunately this exposes a vulnerability where multiple calls to Update can occur, causing it to be destroyed prematurely on the second call (because bNeedsDestroy was made true on the first call), leaving entities still alive referencing that instance with NeedsUnlink until the end of the PostInstantiation.

The fix is to defer the destruction of the instance to the runner itself, like is true for root sequences. Unfortunately this does make FindRelatedInstanceHandle unable to work correctly for expiring sub sequences. In practice this is not really a problem because the only time we use that function is in the template sequence system, and that can reference things just as easily by RootInstanceHandle and SequenceID

diff --git a/Engine/Plugins/MovieScene/TemplateSequence/Source/TemplateSequence/Private/Systems/TemplateSequenceSystem.cpp b/Engine/Plugins/MovieScene/TemplateSequence/Source/TemplateSequence/Private/Systems/TemplateSequenceSystem.cpp
--- a/Engine/Plugins/MovieScene/TemplateSequence/Source/TemplateSequence/Private/Systems/TemplateSequenceSystem.cpp
+++ b/Engine/Plugins/MovieScene/TemplateSequence/Source/TemplateSequence/Private/Systems/TemplateSequenceSystem.cpp
@@ -131,34 +131,32 @@ void UTemplateSequencePropertyScalingInstantiatorSystem::OnRun(FSystemTaskPrereq
 	// Step 1: Keep track of any new property scales.
 	auto GatherNewPropertyScaledInstances = [this, InstanceRegistry](
 			const FMovieSceneEntityID EntityID,
-			const FInstanceHandle InstanceHandle,
+			const FRootInstanceHandle RootInstanceHandle,
 			const FTemplateSequencePropertyScaleComponentData& PropertyScale)
 	{
-		const FInstanceHandle SubSequenceHandle = InstanceRegistry->FindRelatedInstanceHandle(InstanceHandle, PropertyScale.SubSequenceID);
-		if (ensure(SubSequenceHandle.IsValid()))
+		FPropertyScaledInstanceKey Key(RootInstanceHandle, PropertyScale.SubSequenceID);
+
+		FPropertyScaleEntityIDs& PropertyScaleEntities = PropertyScaledInstances.FindOrAdd(Key);
+		if (!PropertyScaleEntities.Contains(EntityID))
 		{
-			FPropertyScaleEntityIDs& PropertyScaleEntities = PropertyScaledInstances.FindOrAdd(SubSequenceHandle);
-			if (!PropertyScaleEntities.Contains(EntityID))
-			{
-				PropertyScaleEntities.Add(EntityID);
+			PropertyScaleEntities.Add(EntityID);
 
-				switch (PropertyScale.PropertyScaleType)
-				{
-					case ETemplateSectionPropertyScaleType::FloatProperty:
-						++FloatScaleUseCount;
-						break;
-					case ETemplateSectionPropertyScaleType::TransformPropertyLocationOnly:
-					case ETemplateSectionPropertyScaleType::TransformPropertyRotationOnly:
-						++TransformScaleUseCount;
-						break;
-				}
+			switch (PropertyScale.PropertyScaleType)
+			{
+				case ETemplateSectionPropertyScaleType::FloatProperty:
+					++FloatScaleUseCount;
+					break;
+				case ETemplateSectionPropertyScaleType::TransformPropertyLocationOnly:
+				case ETemplateSectionPropertyScaleType::TransformPropertyRotationOnly:
+					++TransformScaleUseCount;
+					break;
 			}
 		}
 	};
 
 	FEntityTaskBuilder()
 		.ReadEntityIDs()
-		.Read(BuiltInComponents->InstanceHandle)
+		.Read(BuiltInComponents->RootInstanceHandle)
 		.Read(TemplateSequenceComponents->PropertyScale)
 		.FilterAll({ BuiltInComponents->Tags.NeedsLink })
 		.Iterate_PerEntity(&Linker->EntityManager, GatherNewPropertyScaledInstances);
@@ -166,14 +164,14 @@ void UTemplateSequencePropertyScalingInstantiatorSystem::OnRun(FSystemTaskPrereq
 	// Step 2: Remove old property scales ending.
 	auto RemoveOldPropertyScaledInstances = [this, InstanceRegistry](
 			const FMovieSceneEntityID EntityID,
-			const FInstanceHandle InstanceHandle,
+			const FRootInstanceHandle RootInstanceHandle,
 			const FTemplateSequencePropertyScaleComponentData& PropertyScale)
 	{
-		const FInstanceHandle SubSequenceHandle = InstanceRegistry->FindRelatedInstanceHandle(InstanceHandle, PropertyScale.SubSequenceID);
-		if (ensure(SubSequenceHandle.IsValid()) && PropertyScaledInstances.Contains(SubSequenceHandle))
+		FPropertyScaledInstanceKey Key(RootInstanceHandle, PropertyScale.SubSequenceID);
+
+		if (FPropertyScaleEntityIDs* PropertyScaleEntities = PropertyScaledInstances.Find(Key))
 		{
-			FPropertyScaleEntityIDs& PropertyScaleEntities = PropertyScaledInstances.FindChecked(SubSequenceHandle);
-			if (PropertyScaleEntities.Remove(EntityID) > 0)
+			if (PropertyScaleEntities->Remove(EntityID) > 0)
 			{
 				switch (PropertyScale.PropertyScaleType)
 				{
@@ -186,16 +184,16 @@ void UTemplateSequencePropertyScalingInstantiatorSystem::OnRun(FSystemTaskPrereq
 						break;
 				}
 			}
-			if (PropertyScaleEntities.Num() == 0)
+			if (PropertyScaleEntities->Num() == 0)
 			{
-				PropertyScaledInstances.Remove(SubSequenceHandle);
+				PropertyScaledInstances.Remove(Key);
 			}
 		}
 	};
 
 	FEntityTaskBuilder()
 		.ReadEntityIDs()
-		.Read(BuiltInComponents->InstanceHandle)
+		.Read(BuiltInComponents->RootInstanceHandle)
 		.Read(TemplateSequenceComponents->PropertyScale)
 		.FilterAll({ BuiltInComponents->Tags.NeedsUnlink, TemplateSequenceComponents->PropertyScale })
 		.Iterate_PerEntity(&Linker->EntityManager, RemoveOldPropertyScaledInstances);
@@ -207,11 +205,13 @@ void UTemplateSequencePropertyScalingInstantiatorSystem::OnRun(FSystemTaskPrereq
 
 	auto ComputeReverseLookupBinding = [this, BuiltInComponents, &PropertyScaledEntities](
 			const FMovieSceneEntityID EntityID,
-			const FInstanceHandle InstanceHandle,
+			const FRootInstanceHandle RootInstanceHandle,
+			const FMovieSceneSequenceID* OptSequenceID,
 			const FMovieSceneEntityID ParentEntity, 
 			UObject* const BoundObject)
 	{
-		if (!PropertyScaledInstances.Contains(InstanceHandle))
+		FMovieSceneSequenceID SequenceID = OptSequenceID ? *OptSequenceID : MovieSceneSequenceID::Root;
+		if (!PropertyScaledInstances.Contains(FPropertyScaledInstanceKey(RootInstanceHandle, SequenceID)))
 		{
 			return;
 		}
@@ -235,7 +235,8 @@ void UTemplateSequencePropertyScalingInstantiatorSystem::OnRun(FSystemTaskPrereq
 
 	FEntityTaskBuilder()
 		.ReadEntityIDs()
-		.Read(BuiltInComponents->InstanceHandle)
+		.Read(BuiltInComponents->RootInstanceHandle)
+		.ReadOptional(BuiltInComponents->SequenceID)
 		.Read(BuiltInComponents->ParentEntity)
 		.Read(BuiltInComponents->BoundObject)
 		.FilterAll({ BuiltInComponents->Tags.NeedsLink })
@@ -352,18 +353,14 @@ struct FGatherPropertyScales
 	}
 
 	void ForEachEntity(
-			const FInstanceHandle& InstanceHandle, 
+			const FRootInstanceHandle& RootInstanceHandle, 
 			const FTemplateSequencePropertyScaleComponentData& PropertyScale,
 			float ScaleFactor) const
 	{
 		// Find the instance handle of the sub-sequence we need to scale.
-		const FInstanceHandle SubSequenceHandle = InstanceRegistry->FindRelatedInstanceHandle(InstanceHandle, PropertyScale.SubSequenceID);
-		if (ensure(SubSequenceHandle.IsValid()))
-		{
-			EvaluatorSystem->AddPropertyScale(
-					UEvaluatorSystem::FPropertyScaleKey { SubSequenceHandle, PropertyScale.ObjectBinding, PropertyScale.PropertyBinding.PropertyPath },
-					UEvaluatorSystem::FPropertyScaleValue { PropertyScale.PropertyScaleType, ScaleFactor });
-		}
+		EvaluatorSystem->AddPropertyScale(
+				UEvaluatorSystem::FPropertyScaleKey { RootInstanceHandle, PropertyScale.SubSequenceID, PropertyScale.ObjectBinding, PropertyScale.PropertyBinding.PropertyPath },
+				UEvaluatorSystem::FPropertyScaleValue { PropertyScale.PropertyScaleType, ScaleFactor });
 	}
 };
 
@@ -381,7 +378,8 @@ struct FScaleTransformProperties
 
 	void ForEachAllocation(
 			FEntityAllocation* Allocation,
-			TRead<FInstanceHandle> InstanceHandles,
+			TRead<FRootInstanceHandle> RootInstanceHandles,
+			TReadOptional<FMovieSceneSequenceID> SequenceIDs,
 			TRead<FMovieScenePropertyBinding> PropertyBindings,
 			TRead<FGuid> ReverseBindingLookups,
 			TReadOptional<double> BaseLocationXs, TReadOptional<double> BaseLocationYs, TReadOptional<double> BaseLocationZs,
@@ -395,14 +393,16 @@ struct FScaleTransformProperties
 	{
 		for (int32 Index = 0; Index < Allocation->Num(); ++Index)
 		{
-			const FInstanceHandle InstanceHandle = InstanceHandles[Index];
+			const FRootInstanceHandle InstanceHandle = RootInstanceHandles[Index];
 			const FMovieScenePropertyBinding& PropertyBinding = PropertyBindings[Index];
 			const FGuid ObjectBindingID = ReverseBindingLookups[Index];
+			const FMovieSceneSequenceID SequenceID = SequenceIDs ? SequenceIDs[Index] : MovieSceneSequenceID::Root;
+
 			ensure(ObjectBindingID.IsValid());
 
 			UEvaluatorSystem::FPropertyScaleValueArray TransformScales;
 			EvaluatorSystem->FindPropertyScales(
-					UEvaluatorSystem::FPropertyScaleKey { InstanceHandle, ObjectBindingID, PropertyBinding.PropertyPath },
+					UEvaluatorSystem::FPropertyScaleKey { InstanceHandle, SequenceID, ObjectBindingID, PropertyBinding.PropertyPath },
 					TransformScales);
 			for (const UEvaluatorSystem::FPropertyScaleValue& TransformScale : TransformScales)
 			{
@@ -445,7 +445,8 @@ struct FScaleFloatProperties
 
 	void ForEachAllocation(
 			FEntityAllocation* Allocation,
-			TRead<FInstanceHandle> InstanceHandles,
+			TRead<FRootInstanceHandle> RootInstanceHandles,
+			TReadOptional<FMovieSceneSequenceID> SequenceIDs,
 			TRead<FMovieScenePropertyBinding> PropertyBindings,
 			TRead<FGuid> ReverseBindingLookups,
 			TReadOptional<double> BasePropertyValues,
@@ -453,14 +454,15 @@ struct FScaleFloatProperties
 	{
 		for (int32 Index = 0; Index < Allocation->Num(); ++Index)
 		{
-			const FInstanceHandle InstanceHandle = InstanceHandles[Index];
+			const FRootInstanceHandle InstanceHandle = RootInstanceHandles[Index];
 			const FMovieScenePropertyBinding& PropertyBinding = PropertyBindings[Index];
 			const FGuid ObjectBindingID = ReverseBindingLookups[Index];
+			const FMovieSceneSequenceID SequenceID = SequenceIDs ? SequenceIDs[Index] : MovieSceneSequenceID::Root;
 			ensure(ObjectBindingID.IsValid());
 
 			UEvaluatorSystem::FPropertyScaleValueArray PropertyScales;
 			EvaluatorSystem->FindPropertyScales(
-					UEvaluatorSystem::FPropertyScaleKey { InstanceHandle, ObjectBindingID, PropertyBinding.PropertyPath },
+					UEvaluatorSystem::FPropertyScaleKey { InstanceHandle, SequenceID, ObjectBindingID, PropertyBinding.PropertyPath },
 					PropertyScales);
 			for (const UEvaluatorSystem::FPropertyScaleValue& PropertyScale : PropertyScales)
 			{
@@ -499,7 +501,7 @@ void UTemplateSequencePropertyScalingEvaluatorSystem::OnSchedulePersistentTasks(
 	// - ...are scaled by what factor (scale value), using what type of scale (scale type enum).
 	//
 	FTaskID GatherTask = FEntityTaskBuilder()
-	.Read(BuiltInComponents->InstanceHandle)
+	.Read(BuiltInComponents->RootInstanceHandle)
 	.Read(TemplateSequenceComponents->PropertyScale)
 	.Read(BuiltInComponents->DoubleResult[0])
 	.Schedule_PerEntity<FGatherPropertyScales>(&Linker->EntityManager, TaskScheduler, this);
@@ -510,7 +512,8 @@ void UTemplateSequencePropertyScalingEvaluatorSystem::OnSchedulePersistentTasks(
 	if (ensure(InstantiatorSystem) && InstantiatorSystem->HasAnyTransformScales())
 	{
 		FTaskID ScaleTransformTask = FEntityTaskBuilder()
-			.Read(BuiltInComponents->InstanceHandle)
+			.Read(BuiltInComponents->RootInstanceHandle)
+			.ReadOptional(BuiltInComponents->SequenceID)
 			.Read(BuiltInComponents->PropertyBinding)
 			.Read(TemplateSequenceComponents->PropertyScaleReverseBindingLookup)
 			.ReadOptional(BuiltInComponents->BaseDouble[0])
@@ -535,7 +538,8 @@ void UTemplateSequencePropertyScalingEvaluatorSystem::OnSchedulePersistentTasks(
 	if (ensure(InstantiatorSystem) && InstantiatorSystem->HasAnyFloatScales())
 	{
 		FTaskID ScaleFloatTask = FEntityTaskBuilder()
-			.Read(BuiltInComponents->InstanceHandle)
+			.Read(BuiltInComponents->RootInstanceHandle)
+			.ReadOptional(BuiltInComponents->SequenceID)
 			.Read(BuiltInComponents->PropertyBinding)
 			.Read(TemplateSequenceComponents->PropertyScaleReverseBindingLookup)
 			.ReadOptional(BuiltInComponents->BaseDouble[0])
@@ -567,7 +571,7 @@ void UTemplateSequencePropertyScalingEvaluatorSystem::OnRun(FSystemTaskPrerequis
 	// - ...are scaled by what factor (scale value), using what type of scale (scale type enum).
 	//
 	FGraphEventRef GatherTask = FEntityTaskBuilder()
-		.Read(BuiltInComponents->InstanceHandle)
+		.Read(BuiltInComponents->RootInstanceHandle)
 		.Read(TemplateSequenceComponents->PropertyScale)
 		.Read(BuiltInComponents->DoubleResult[0])
 		.Dispatch_PerEntity<FGatherPropertyScales>(&Linker->EntityManager, InPrerequisites, &Subsequents, this);
@@ -581,7 +585,8 @@ void UTemplateSequencePropertyScalingEvaluatorSystem::OnRun(FSystemTaskPrerequis
 	if (ensure(InstantiatorSystem) && InstantiatorSystem->HasAnyTransformScales())
 	{
 		FEntityTaskBuilder()
-			.Read(BuiltInComponents->InstanceHandle)
+			.Read(BuiltInComponents->RootInstanceHandle)
+			.ReadOptional(BuiltInComponents->SequenceID)
 			.Read(BuiltInComponents->PropertyBinding)
 			.Read(TemplateSequenceComponents->PropertyScaleReverseBindingLookup)
 			.ReadOptional(BuiltInComponents->BaseDouble[0])
@@ -604,7 +609,8 @@ void UTemplateSequencePropertyScalingEvaluatorSystem::OnRun(FSystemTaskPrerequis
 	if (ensure(InstantiatorSystem) && InstantiatorSystem->HasAnyFloatScales())
 	{
 		FEntityTaskBuilder()
-			.Read(BuiltInComponents->InstanceHandle)
+			.Read(BuiltInComponents->RootInstanceHandle)
+			.ReadOptional(BuiltInComponents->SequenceID)
 			.Read(BuiltInComponents->PropertyBinding)
 			.Read(TemplateSequenceComponents->PropertyScaleReverseBindingLookup)
 			.ReadOptional(BuiltInComponents->BaseDouble[0])
diff --git a/Engine/Plugins/MovieScene/TemplateSequence/Source/TemplateSequence/Private/Systems/TemplateSequenceSystem.h b/Engine/Plugins/MovieScene/TemplateSequence/Source/TemplateSequence/Private/Systems/TemplateSequenceSystem.h
--- a/Engine/Plugins/MovieScene/TemplateSequence/Source/TemplateSequence/Private/Systems/TemplateSequenceSystem.h
+++ b/Engine/Plugins/MovieScene/TemplateSequence/Source/TemplateSequence/Private/Systems/TemplateSequenceSystem.h
@@ -42,8 +42,10 @@ private:
 
 private:
 	
+	using FPropertyScaledInstanceKey = TTuple<UE::MovieScene::FRootInstanceHandle, FMovieSceneSequenceID>;
+
 	using FPropertyScaleEntityIDs = TArray<UE::MovieScene::FMovieSceneEntityID, TInlineAllocator<2>>;
-	TMap<UE::MovieScene::FInstanceHandle, FPropertyScaleEntityIDs> PropertyScaledInstances;
+	TMap<FPropertyScaledInstanceKey, FPropertyScaleEntityIDs> PropertyScaledInstances;
 
 	int32 FloatScaleUseCount = 0;
 	int32 TransformScaleUseCount = 0;
@@ -58,7 +60,7 @@ public:
 
 	UTemplateSequencePropertyScalingEvaluatorSystem(const FObjectInitializer& ObjInit);
 
-	using FPropertyScaleKey = TTuple<UE::MovieScene::FInstanceHandle, FGuid, FName>;
+	using FPropertyScaleKey = TTuple<UE::MovieScene::FRootInstanceHandle, FMovieSceneSequenceID, FGuid, FName>;
 	using FPropertyScaleValue = TTuple<ETemplateSectionPropertyScaleType, float>;
 
 	using FPropertyScaleValueArray = TArray<FPropertyScaleValue, TInlineAllocator<2>>;
diff --git a/Engine/Source/Runtime/MovieScene/Private/EntitySystem/MovieSceneSequenceUpdaters.cpp b/Engine/Source/Runtime/MovieScene/Private/EntitySystem/MovieSceneSequenceUpdaters.cpp
--- a/Engine/Source/Runtime/MovieScene/Private/EntitySystem/MovieSceneSequenceUpdaters.cpp
+++ b/Engine/Source/Runtime/MovieScene/Private/EntitySystem/MovieSceneSequenceUpdaters.cpp
@@ -90,7 +90,6 @@ private:
 	{
 		FGuid SequenceSignature;
 		FInstanceHandle Handle;
-		bool bNeedsDestroy = false;
 	};
 	TSortedMap<FMovieSceneSequenceID, FSubInstanceData, TInlineAllocator<8>> SequenceInstances;
 
@@ -464,7 +463,6 @@ FInstanceHandle FSequenceUpdater_Hierarchical::GetOrCreateSequenceInstance(IMovi
 
 	if (FSubInstanceData* Existing = SequenceInstances.Find(SequenceID))
 	{
-		Existing->bNeedsDestroy = false;
 		return Existing->Handle;
 	}
 
@@ -734,21 +732,15 @@ void FSequenceUpdater_Hierarchical::Update(UMovieSceneEntitySystemLinker* Linker
 	for (auto InstanceIt = SequenceInstances.CreateIterator(); InstanceIt; ++InstanceIt)
 	{
 		FSubInstanceData SubData = InstanceIt.Value();
-		if (SubData.bNeedsDestroy)
-		{
-			InstanceRegistry->DestroyInstance(SubData.Handle);
-			InstanceIt.RemoveCurrent();
-			continue;
-		}
-
-		Runner->MarkForUpdate(SubData.Handle, ERunnerUpdateFlags::None);
 
+		ERunnerUpdateFlags Flags = ERunnerUpdateFlags::None;
 		if (!ActiveSequences.Contains(InstanceIt.Key()))
 		{
-			// Remove all entities from this instance since it is no longer active
-			InstanceRegistry->MutateInstance(SubData.Handle).Finish(Linker);
-			InstanceIt.Value().bNeedsDestroy = true;
+			Flags = ERunnerUpdateFlags::Finish | ERunnerUpdateFlags::Destroy;
+			InstanceIt.RemoveCurrent();
 		}
+
+		Runner->MarkForUpdate(SubData.Handle, Flags);
 	}
 }
 
