Sequencer: Allow resizing sub track and template sequence track row heights

diff --git a/Engine/Plugins/MovieScene/TemplateSequence/Source/TemplateSequenceEditor/Private/TrackEditors/TemplateSequenceTrackEditor.cpp b/Engine/Plugins/MovieScene/TemplateSequence/Source/TemplateSequenceEditor/Private/TrackEditors/TemplateSequenceTrackEditor.cpp
--- a/Engine/Plugins/MovieScene/TemplateSequence/Source/TemplateSequenceEditor/Private/TrackEditors/TemplateSequenceTrackEditor.cpp
+++ b/Engine/Plugins/MovieScene/TemplateSequence/Source/TemplateSequenceEditor/Private/TrackEditors/TemplateSequenceTrackEditor.cpp
@@ -18,6 +18,7 @@
 #include "Tracks/MovieSceneFloatTrack.h"
 #include "ScopedTransaction.h"
 #include "Tracks/MovieScene3DTransformTrack.h"
+#include "Tracks/MovieSceneSubTrack.h"
 #include "Sections/TemplateSequenceSection.h"
 #include "Tracks/TemplateSequenceTrack.h"
 #include "Widgets/Input/SButton.h"
@@ -445,6 +446,28 @@ UCameraComponent* FTemplateSequenceTrackEditor::AcquireCameraComponentFromObject
 	return nullptr;
 }
 
+bool FTemplateSequenceTrackEditor::IsResizable(UMovieSceneTrack* InTrack) const
+{
+	return true;
+}
+
+void FTemplateSequenceTrackEditor::Resize(float NewSize, UMovieSceneTrack* InTrack)
+{
+	UMovieSceneSubTrack* SubTrack = Cast<UMovieSceneSubTrack>(InTrack);
+	if (SubTrack)
+	{
+		SubTrack->Modify();
+
+		int32 MaxNumRows = 1;
+		for (UMovieSceneSection* Section : SubTrack->GetAllSections())
+		{
+			MaxNumRows = FMath::Max(MaxNumRows, Section->GetRowIndex() + 1);
+		}
+
+		SubTrack->SetRowHeight(FMath::RoundToInt(NewSize) / MaxNumRows);
+	}
+}
+
 FTemplateSequenceSection::FTemplateSequenceSection(TSharedPtr<ISequencer> InSequencer, UTemplateSequenceSection& InSection)
 	: TSubSectionMixin(InSequencer, InSection)
 {
diff --git a/Engine/Plugins/MovieScene/TemplateSequence/Source/TemplateSequenceEditor/Private/TrackEditors/TemplateSequenceTrackEditor.h b/Engine/Plugins/MovieScene/TemplateSequence/Source/TemplateSequenceEditor/Private/TrackEditors/TemplateSequenceTrackEditor.h
--- a/Engine/Plugins/MovieScene/TemplateSequence/Source/TemplateSequenceEditor/Private/TrackEditors/TemplateSequenceTrackEditor.h
+++ b/Engine/Plugins/MovieScene/TemplateSequence/Source/TemplateSequenceEditor/Private/TrackEditors/TemplateSequenceTrackEditor.h
@@ -29,6 +29,8 @@ public:
 	virtual void BuildObjectBindingTrackMenu(FMenuBuilder& MenuBuilder, const TArray<FGuid>& ObjectBindings, const UClass* ObjectClass) override;
 	TSharedPtr<SWidget> BuildOutlinerEditWidget(const FGuid& ObjectBinding, UMovieSceneTrack* Track, const FBuildEditWidgetParams& Params) override;
 	virtual TSharedRef<ISequencerSection> MakeSectionInterface(UMovieSceneSection& SectionObject, UMovieSceneTrack& Track, FGuid ObjectBinding) override;
+	virtual bool IsResizable(UMovieSceneTrack* InTrack) const override;
+	virtual void Resize(float NewSize, UMovieSceneTrack* InTrack) override;
 
 private:
 	void AddTemplateSequenceAssetSubMenu(FMenuBuilder& MenuBuilder, TArray<FGuid> ObjectBindings, const UClass* RootBindingClass);
diff --git a/Engine/Source/Editor/MovieSceneTools/Private/TrackEditors/AudioTrackEditor.cpp b/Engine/Source/Editor/MovieSceneTools/Private/TrackEditors/AudioTrackEditor.cpp
--- a/Engine/Source/Editor/MovieSceneTools/Private/TrackEditors/AudioTrackEditor.cpp
+++ b/Engine/Source/Editor/MovieSceneTools/Private/TrackEditors/AudioTrackEditor.cpp
@@ -1005,12 +1005,7 @@ void FAudioTrackEditor::Resize(float NewSize, UMovieSceneTrack* InTrack)
 	{
 		AudioTrack->Modify();
 
-		int32 MaxNumRows = 1;
-		for (UMovieSceneSection* Section : AudioTrack->GetAllSections())
-		{
-			MaxNumRows = FMath::Max(MaxNumRows, Section->GetRowIndex() + 1);
-		}
-
+		const int32 MaxNumRows = AudioTrack->GetMaxRowIndex() + 1;
 		AudioTrack->SetRowHeight(FMath::RoundToInt(NewSize) / MaxNumRows);
 	}
 }
diff --git a/Engine/Source/Editor/MovieSceneTools/Private/TrackEditors/SubTrackEditor.cpp b/Engine/Source/Editor/MovieSceneTools/Private/TrackEditors/SubTrackEditor.cpp
--- a/Engine/Source/Editor/MovieSceneTools/Private/TrackEditors/SubTrackEditor.cpp
+++ b/Engine/Source/Editor/MovieSceneTools/Private/TrackEditors/SubTrackEditor.cpp
@@ -37,12 +37,6 @@
 #include "IStructureDetailsView.h"
 #include "PropertyEditorModule.h"
 
-namespace SubTrackEditorConstants
-{
-	const float TrackHeight = 50.0f;
-}
-
-
 #define LOCTEXT_NAMESPACE "FSubTrackEditor"
 
 
@@ -64,11 +58,6 @@ public:
 
 	// ISequencerSection interface
 
-	virtual float GetSectionHeight() const override
-	{
-		return SubTrackEditorConstants::TrackHeight;
-	}
-	
 	virtual void BuildSectionContextMenu(FMenuBuilder& MenuBuilder, const FGuid& ObjectBinding) override
 	{
 		ISequencerSection::BuildSectionContextMenu(MenuBuilder, ObjectBinding);
@@ -443,6 +432,24 @@ FReply FSubTrackEditor::OnDrop(const FDragDropEvent& DragDropEvent, const FSeque
 	return bAnyDropped ? FReply::Handled() : FReply::Unhandled();
 }
 
+bool FSubTrackEditor::IsResizable(UMovieSceneTrack* InTrack) const
+{
+	return true;
+}
+
+void FSubTrackEditor::Resize(float NewSize, UMovieSceneTrack* InTrack)
+{
+	UMovieSceneSubTrack* SubTrack = Cast<UMovieSceneSubTrack>(InTrack);
+	if (SubTrack)
+	{
+		SubTrack->Modify();
+
+		const int32 MaxNumRows = SubTrack->GetMaxRowIndex() + 1;
+		SubTrack->SetRowHeight(FMath::RoundToInt(NewSize) / MaxNumRows);
+		SubTrack->SetRowHeight(NewSize);
+	}
+}
+
 /* FSubTrackEditor
  *****************************************************************************/
 
diff --git a/Engine/Source/Editor/MovieSceneTools/Private/TrackEditors/SubTrackEditorBase.cpp b/Engine/Source/Editor/MovieSceneTools/Private/TrackEditors/SubTrackEditorBase.cpp
--- a/Engine/Source/Editor/MovieSceneTools/Private/TrackEditors/SubTrackEditorBase.cpp
+++ b/Engine/Source/Editor/MovieSceneTools/Private/TrackEditors/SubTrackEditorBase.cpp
@@ -97,7 +97,11 @@ FSubSectionPainterResult FSubSectionPainterUtil::PaintSection(TSharedPtr<const I
 
     if (!SectionText.IsEmpty())
     {
-        FSlateDrawElement::MakeText(
+		FSlateClippingZone ClippingZone(InPainter.SectionClippingRect.InsetBy(FMargin(1.0f)));
+
+		InPainter.DrawElements.PushClip(ClippingZone);
+		
+		FSlateDrawElement::MakeText(
             InPainter.DrawElements,
             ++LayerId,
             InPainter.SectionGeometry.ToPaintGeometry(
@@ -109,7 +113,9 @@ FSubSectionPainterResult FSubSectionPainterUtil::PaintSection(TSharedPtr<const I
             DrawEffects,
             FColor(200, 200, 200, static_cast<uint8>(255 * InPainter.GhostAlpha))
         );
-    }
+
+		InPainter.DrawElements.PopClip();
+	}
 
     InPainter.LayerId = LayerId;
 
diff --git a/Engine/Source/Editor/MovieSceneTools/Public/TrackEditors/SubTrackEditor.h b/Engine/Source/Editor/MovieSceneTools/Public/TrackEditors/SubTrackEditor.h
--- a/Engine/Source/Editor/MovieSceneTools/Public/TrackEditors/SubTrackEditor.h
+++ b/Engine/Source/Editor/MovieSceneTools/Public/TrackEditors/SubTrackEditor.h
@@ -56,7 +56,9 @@ public:
 	virtual const FSlateBrush* GetIconBrush() const override;
 	virtual bool OnAllowDrop(const FDragDropEvent& DragDropEvent, FSequencerDragDropParams& DragDropParams) override;
 	virtual FReply OnDrop(const FDragDropEvent& DragDropEvent, const FSequencerDragDropParams& DragDropParams) override;
-	
+	virtual bool IsResizable(UMovieSceneTrack* InTrack) const override;
+	virtual void Resize(float NewSize, UMovieSceneTrack* InTrack) override;
+
 public:
 	
 	/** Insert sequence into this track */
diff --git a/Engine/Source/Editor/MovieSceneTools/Public/TrackEditors/SubTrackEditorBase.h b/Engine/Source/Editor/MovieSceneTools/Public/TrackEditors/SubTrackEditorBase.h
--- a/Engine/Source/Editor/MovieSceneTools/Public/TrackEditors/SubTrackEditorBase.h
+++ b/Engine/Source/Editor/MovieSceneTools/Public/TrackEditors/SubTrackEditorBase.h
@@ -18,6 +18,7 @@
 #include "Rendering/DrawElements.h"
 #include "SequencerSectionPainter.h"
 #include "Subsystems/AssetEditorSubsystem.h"
+#include "Tracks/MovieSceneSubTrack.h"
 
 class ISequencer;
 
@@ -289,6 +290,10 @@ TOptional<FFrameTime> TSubSectionMixin<ParentSectionClass>::GetSectionTime(FSequ
 template<typename ParentSectionClass>
 float TSubSectionMixin<ParentSectionClass>::GetSectionHeight() const
 {
+	if (UMovieSceneSubTrack* Track = SubSectionObject.GetTypedOuter<UMovieSceneSubTrack>())
+	{
+		return Track->GetRowHeight();
+	}
 	return TSubSectionMixin<ParentSectionClass>::TrackHeight;
 }
 
diff --git a/Engine/Source/Editor/SequencerCore/Private/MVVM/Views/STrackLane.cpp b/Engine/Source/Editor/SequencerCore/Private/MVVM/Views/STrackLane.cpp
--- a/Engine/Source/Editor/SequencerCore/Private/MVVM/Views/STrackLane.cpp
+++ b/Engine/Source/Editor/SequencerCore/Private/MVVM/Views/STrackLane.cpp
@@ -492,7 +492,7 @@ FReply STrackLane::OnMouseButtonDown(const FGeometry& MyGeometry, const FPointer
 				TViewModelPtr<IOutlinerExtension> OutlinerItem = WeakOutlinerItem.Pin();
 				const float OriginalHeight = OutlinerItem ? OutlinerItem->GetOutlinerSizing().GetTotalHeight() : 10.f;
 
-				DragParameters = FDragParameters(OriginalHeight, MouseEvent.GetScreenSpacePosition().Y);
+				DragParameters = FDragParameters(OriginalHeight, LocalPos.Y);
 				return FReply::Handled().CaptureMouse(AsShared());
 			}
 		}
@@ -517,7 +517,9 @@ FReply STrackLane::OnMouseMove(const FGeometry& MyGeometry, const FPointerEvent&
 
 	if (DragParameters.IsSet() && HasMouseCapture() && ResizableExtension)
 	{
-		float NewHeight = DragParameters->OriginalHeight + (MouseEvent.GetScreenSpacePosition().Y - DragParameters->DragStartY);
+		FVector2D LocalPos = MyGeometry.AbsoluteToLocal(MouseEvent.GetScreenSpacePosition());
+
+		float NewHeight = DragParameters->OriginalHeight + (LocalPos.Y - DragParameters->DragStartY);
 
 		if (FMath::RoundToInt(NewHeight) != FMath::RoundToInt(DragParameters->OriginalHeight))
 		{
diff --git a/Engine/Source/Runtime/MovieScene/Private/Tracks/MovieSceneSubTrack.cpp b/Engine/Source/Runtime/MovieScene/Private/Tracks/MovieSceneSubTrack.cpp
--- a/Engine/Source/Runtime/MovieScene/Private/Tracks/MovieSceneSubTrack.cpp
+++ b/Engine/Source/Runtime/MovieScene/Private/Tracks/MovieSceneSubTrack.cpp
@@ -22,6 +22,7 @@ UMovieSceneSubTrack::UMovieSceneSubTrack( const FObjectInitializer& ObjectInitia
 {
 #if WITH_EDITORONLY_DATA
 	TrackTint = FColor(180, 0, 40, 65);
+	RowHeight = 50;
 #endif
 
 	BuiltInTreePopulationMode = ETreePopulationMode::Blended;
diff --git a/Engine/Source/Runtime/MovieScene/Public/Tracks/MovieSceneSubTrack.h b/Engine/Source/Runtime/MovieScene/Public/Tracks/MovieSceneSubTrack.h
--- a/Engine/Source/Runtime/MovieScene/Public/Tracks/MovieSceneSubTrack.h
+++ b/Engine/Source/Runtime/MovieScene/Public/Tracks/MovieSceneSubTrack.h
@@ -90,4 +90,32 @@ protected:
 	/** All movie scene sections. */
 	UPROPERTY()
 	TArray<TObjectPtr<UMovieSceneSection>> Sections;
+
+#if WITH_EDITORONLY_DATA
+
+public:
+
+	/**
+	 * Get the height of this track's rows
+	 */
+	int32 GetRowHeight() const
+	{
+		return RowHeight;
+	}
+
+	/**
+	 * Set the height of this track's rows
+	 */
+	void SetRowHeight(int32 NewRowHeight)
+	{
+		RowHeight = FMath::Max(16, NewRowHeight);
+	}
+
+private:
+
+	/** The height for each row of this track */
+	UPROPERTY()
+	int32 RowHeight;
+
+#endif
 };
