Sequencer: don't duplicate weighted blendables when a camera animation includes them

diff --git a/Engine/Plugins/MovieScene/TemplateSequence/Source/TemplateSequence/Private/CameraAnimationSequencePlayer.cpp b/Engine/Plugins/MovieScene/TemplateSequence/Source/TemplateSequence/Private/CameraAnimationSequencePlayer.cpp
--- a/Engine/Plugins/MovieScene/TemplateSequence/Source/TemplateSequence/Private/CameraAnimationSequencePlayer.cpp
+++ b/Engine/Plugins/MovieScene/TemplateSequence/Source/TemplateSequence/Private/CameraAnimationSequencePlayer.cpp
@@ -172,13 +172,20 @@ void UCameraAnimationSequenceCameraStandIn::ResetDefaultValues(const FMinimalVie
 	CurrentFocalLength = (Filmback.SensorWidth / 2.f) / FMath::Tan(FMath::DegreesToRadians(FieldOfView / 2.f));
 
 	// Restore weighted blendables.
+	// Either the camera animation did not add any blendables, in which case we just restore what
+	// was there before, or it did add some, in which case we restore anything *else* that was
+	// there before. This assumes that there can't be two instances of the same blendable object,
+	// but the AddBlendable method (which we use here) already makes this assumption anyway.
 	if (PostProcessSettings.WeightedBlendables.Array.Num() == 0)
 	{
 		PostProcessSettings.WeightedBlendables.Array = MoveTemp(WBBackup);
 	}
 	else
 	{
-		PostProcessSettings.WeightedBlendables.Array.Append(WBBackup);
+		for (const FWeightedBlendable& WB : WBBackup)
+		{
+			PostProcessSettings.AddBlendable(WB.Object, WB.Weight);
+		}
 	}
 
 	RecalcDerivedData();
