Sequencer: fix assert firing when scrubbing backwards into a camera shake whose length is less than that of its shake section

When doing so, the player has to restart playing. This change always restarts it -- if the current eval time falls outside of the shake's bounds, it will re-stop by itself.

diff --git a/Engine/Plugins/MovieScene/TemplateSequence/Source/TemplateSequence/Private/SequenceCameraShake.cpp b/Engine/Plugins/MovieScene/TemplateSequence/Source/TemplateSequence/Private/SequenceCameraShake.cpp
--- a/Engine/Plugins/MovieScene/TemplateSequence/Source/TemplateSequence/Private/SequenceCameraShake.cpp
+++ b/Engine/Plugins/MovieScene/TemplateSequence/Source/TemplateSequence/Private/SequenceCameraShake.cpp
@@ -108,6 +108,11 @@ void USequenceCameraShakePattern::ScrubShakePatternImpl(const FCameraShakePatter
 	const float BlendWeight = State.Scrub(Params.AbsoluteTime);
 	if (State.IsPlaying())
 	{
+		if (Player->GetPlaybackStatus() != EMovieScenePlayerStatus::Playing)
+		{
+			Player->Play(bRandomSegment, bRandomSegment);
+		}
+
 		const FFrameRate InputRate = Player->GetInputRate();
 		const FFrameTime NewPosition = Params.AbsoluteTime * PlayRate * InputRate;
 		UpdateCamera(NewPosition, Params.POV, OutResult);
