Sequencer: handle the case of the camera animation linker being destroyed while animations are still playing

When doing a seamless travel from one level to the next, the camera animation linker (used by camera animations and sequence shakes) can be destroyed earlier than the player camera manager. If camera animations/shakes are playing, they could be ticked a few times after the linker is gone, resulting in a crash when the player tries to do anything.

diff --git a/Engine/Plugins/Cameras/GameplayCameras/Source/GameplayCameras/Private/Animations/CameraAnimationCameraModifier.cpp b/Engine/Plugins/Cameras/GameplayCameras/Source/GameplayCameras/Private/Animations/CameraAnimationCameraModifier.cpp
--- a/Engine/Plugins/Cameras/GameplayCameras/Source/GameplayCameras/Private/Animations/CameraAnimationCameraModifier.cpp
+++ b/Engine/Plugins/Cameras/GameplayCameras/Source/GameplayCameras/Private/Animations/CameraAnimationCameraModifier.cpp
@@ -36,6 +36,16 @@ FActiveCameraAnimationInfo::FActiveCameraAnimationInfo()
 {
 }
 
+bool FActiveCameraAnimationInfo::IsValid() const
+{
+	return Sequence != nullptr;
+}
+
+bool FActiveCameraAnimationInfo::HasValidPlayer() const
+{
+	return Player && Player->IsValid();
+}
+
 namespace CameraAnimationCameraModifierImpl
 {
 	float Sinusoidal(float InTime)
@@ -305,11 +315,22 @@ void UCameraAnimationCameraModifier::TickAllAnimations(float DeltaTime, FMinimal
 	{
 		if (ActiveAnimation.IsValid())
 		{
-			TickAnimation(ActiveAnimation, DeltaTime, InOutPOV);
+			if (ActiveAnimation.HasValidPlayer())
+			{
+				TickAnimation(ActiveAnimation, DeltaTime, InOutPOV);
 
-			if (ActiveAnimation.Player->GetPlaybackStatus() == EMovieScenePlayerStatus::Stopped)
+				if (ActiveAnimation.Player->GetPlaybackStatus() == EMovieScenePlayerStatus::Stopped)
+				{
+					DeactivateCameraAnimation(ActiveAnimation.Handle.InstanceID);
+				}
+			}
+			else
 			{
-				DeactivateCameraAnimation(ActiveAnimation.Handle.InstanceID);
+				// The animation's player might have become invalid due to things like server travel 
+				// to a new level. The player's linker would have been destroyed with the old level.
+				// In that case, we abort the camera animation.
+				ActiveAnimation.Sequence = nullptr;
+				ActiveAnimation.Player = nullptr;
 			}
 		}
 	}
diff --git a/Engine/Plugins/Cameras/GameplayCameras/Source/GameplayCameras/Public/Animations/CameraAnimationCameraModifier.h b/Engine/Plugins/Cameras/GameplayCameras/Source/GameplayCameras/Public/Animations/CameraAnimationCameraModifier.h
--- a/Engine/Plugins/Cameras/GameplayCameras/Source/GameplayCameras/Public/Animations/CameraAnimationCameraModifier.h
+++ b/Engine/Plugins/Cameras/GameplayCameras/Source/GameplayCameras/Public/Animations/CameraAnimationCameraModifier.h
@@ -136,7 +136,10 @@ struct GAMEPLAYCAMERAS_API FActiveCameraAnimationInfo
 	FActiveCameraAnimationInfo();
 
 	/** Whether this is a valid, ongoing camera animation */
-	bool IsValid() const { return Sequence != nullptr; }
+	bool IsValid() const;
+
+	/** Whether this camera animation's player is valid */
+	bool HasValidPlayer() const;
 
 	/** The sequence to use for the animation. */
 	UPROPERTY()
diff --git a/Engine/Plugins/MovieScene/TemplateSequence/Source/TemplateSequence/Private/CameraAnimationSequencePlayer.cpp b/Engine/Plugins/MovieScene/TemplateSequence/Source/TemplateSequence/Private/CameraAnimationSequencePlayer.cpp
--- a/Engine/Plugins/MovieScene/TemplateSequence/Source/TemplateSequence/Private/CameraAnimationSequencePlayer.cpp
+++ b/Engine/Plugins/MovieScene/TemplateSequence/Source/TemplateSequence/Private/CameraAnimationSequencePlayer.cpp
@@ -319,6 +319,11 @@ void UCameraAnimationSequencePlayer::Initialize(UMovieSceneSequence* InSequence,
 	if (Sequence)
 	{
 		Stop();
+
+		if (UMovieSceneEntitySystemLinker* Linker = RootTemplateInstance.GetEntitySystemLinker())
+		{
+			Linker->Events.AbandonLinker.RemoveAll(this);
+		}
 	}
 
 	Sequence = InSequence;
@@ -382,6 +387,23 @@ void UCameraAnimationSequencePlayer::Initialize(UMovieSceneSequence* InSequence,
 	UCameraAnimationSequenceSubsystem* Subsystem = UCameraAnimationSequenceSubsystem::GetCameraAnimationSequenceSubsystem(GetWorld());
 	ensureMsgf(Subsystem, TEXT("Unable to locate a valid camera animation sub-system. Camera anim sequences will not play."));
 	RootTemplateInstance.Initialize(*Sequence, *this, nullptr);
+
+	if (UMovieSceneEntitySystemLinker* Linker = RootTemplateInstance.GetEntitySystemLinker())
+	{
+		Linker->Events.AbandonLinker.AddUObject(this, &UCameraAnimationSequencePlayer::OnAbandonLinker);
+	}
+}
+
+void UCameraAnimationSequencePlayer::OnAbandonLinker(UMovieSceneEntitySystemLinker* InLinker)
+{
+	// Camera animations and camera shakes playing camera animations can outlive the
+	// level in which the linker lives. In this case we just stop.
+	RootTemplateInstance.TearDown();
+}
+
+bool UCameraAnimationSequencePlayer::IsValid() const
+{
+	return RootTemplateInstance.IsValid();
 }
 
 void UCameraAnimationSequencePlayer::Play(bool bLoop, bool bRandomStartTime)
diff --git a/Engine/Plugins/MovieScene/TemplateSequence/Source/TemplateSequence/Private/SequenceCameraShake.cpp b/Engine/Plugins/MovieScene/TemplateSequence/Source/TemplateSequence/Private/SequenceCameraShake.cpp
--- a/Engine/Plugins/MovieScene/TemplateSequence/Source/TemplateSequence/Private/SequenceCameraShake.cpp
+++ b/Engine/Plugins/MovieScene/TemplateSequence/Source/TemplateSequence/Private/SequenceCameraShake.cpp
@@ -169,10 +169,14 @@ void USequenceCameraShakePattern::UpdateCamera(FFrameTime NewPosition, const FMi
 		return;
 	}
 
+	if (!ensure(Player) || !Player->IsValid())
+	{
+		return;
+	}
+
 	using namespace UE::MovieScene;
 
 	check(CameraStandIn);
-	check(Player);
 
 	UMovieSceneEntitySystemLinker* Linker = Player->GetEvaluationTemplate().GetEntitySystemLinker();
 	CameraStandIn->Reset(InPOV, Linker);
diff --git a/Engine/Plugins/MovieScene/TemplateSequence/Source/TemplateSequence/Public/CameraAnimationSequencePlayer.h b/Engine/Plugins/MovieScene/TemplateSequence/Source/TemplateSequence/Public/CameraAnimationSequencePlayer.h
--- a/Engine/Plugins/MovieScene/TemplateSequence/Source/TemplateSequence/Public/CameraAnimationSequencePlayer.h
+++ b/Engine/Plugins/MovieScene/TemplateSequence/Source/TemplateSequence/Public/CameraAnimationSequencePlayer.h
@@ -118,6 +118,12 @@ public:
 	 */
 	void Initialize(UMovieSceneSequence* InSequence, int32 StartOffset = 0, float DurationOverride = 0.f);
 
+	/**
+	 * Returns whether this player is still valid. It could be forcibly torn-down if the camera
+	 * animation subsystem was destroyed with the world.
+	 */
+	bool IsValid() const;
+
 	/**
 	 * Start playing the sequence
 	 *
@@ -190,8 +196,11 @@ public:
 	virtual void BeginDestroy() override;
 
 protected:
+
 	virtual void ResolveBoundObjects(UE::UniversalObjectLocator::FResolveParams& ResolveParams, const FGuid& InBindingId, FMovieSceneSequenceID SequenceID, UMovieSceneSequence& Sequence, TArray<UObject*, TInlineAllocator<1>>& OutObjects) const override;
 
+	void OnAbandonLinker(UMovieSceneEntitySystemLinker* InLinker);
+
 private:
 
 	FSequenceCameraShakeSpawnRegister SpawnRegister;
