[MovieScene]
* Fixed another non unity compile error

diff --git a/Engine/Plugins/MovieScene/TemplateSequence/Source/TemplateSequence/Public/TemplateSequenceComponentTypes.h b/Engine/Plugins/MovieScene/TemplateSequence/Source/TemplateSequence/Public/TemplateSequenceComponentTypes.h
--- a/Engine/Plugins/MovieScene/TemplateSequence/Source/TemplateSequence/Public/TemplateSequenceComponentTypes.h
+++ b/Engine/Plugins/MovieScene/TemplateSequence/Source/TemplateSequence/Public/TemplateSequenceComponentTypes.h
@@ -2,6 +2,7 @@
 
 #pragma once
 
+#include "EntitySystem/MovieSceneEntityIDs.h"
 #include "EntitySystem/MovieScenePropertyBinding.h"
 #include "Evaluation/MovieSceneEvaluationOperand.h"
 
