Sequencer: add support for playing back camera animations with an editor viewport offset

This change adds a new toolbar button for camera animation sequences. The button sets the "base" transform of the camera to wherever the viewport is, so that additive animations play in a more useful place than at the origin.

There is also a new menu item to define whether the viewport transform is automatically used when opening the camera animation sequence in the first place.

diff --git a/Engine/Plugins/MovieScene/TemplateSequence/Source/TemplateSequenceEditor/Private/Customizations/CameraAnimationSequenceCustomization.cpp b/Engine/Plugins/MovieScene/TemplateSequence/Source/TemplateSequenceEditor/Private/Customizations/CameraAnimationSequenceCustomization.cpp
--- a/Engine/Plugins/MovieScene/TemplateSequence/Source/TemplateSequenceEditor/Private/Customizations/CameraAnimationSequenceCustomization.cpp
+++ b/Engine/Plugins/MovieScene/TemplateSequence/Source/TemplateSequenceEditor/Private/Customizations/CameraAnimationSequenceCustomization.cpp
@@ -4,7 +4,9 @@
 #include "CineCameraActor.h"
 #include "Customizations/TemplateSequenceCustomizationBase.h"
 #include "Framework/MultiBox/MultiBoxBuilder.h"
+#include "Misc/TemplateSequenceEditorSettings.h"
 #include "Styles/TemplateSequenceEditorStyle.h"
+#include "Systems/TemplateSequenceCameraPreviewSystem.h"
 
 #define LOCTEXT_NAMESPACE "CameraAnimationSequenceCustomization"
 
@@ -17,9 +19,13 @@ void FCameraAnimationSequenceCustomization::RegisterSequencerCustomization(FSequ
 	FSequencerCustomizationInfo Customization;
 
 	TSharedRef<FExtender> ToolbarExtender = MakeShared<FExtender>();
-	ToolbarExtender->AddToolBarExtension("BaseCommands", EExtensionHook::After, nullptr, FToolBarExtensionDelegate::CreateRaw(this, &FCameraAnimationSequenceCustomization::ExtendSequencerToolbar));
+	ToolbarExtender->AddToolBarExtension("BaseCommands", EExtensionHook::After, nullptr, FToolBarExtensionDelegate::CreateRaw(this, &FCameraAnimationSequenceCustomization::ExtendToolbar));
 	Customization.ToolbarExtender = ToolbarExtender;
 
+	TSharedRef<FExtender> ViewMenuExtender = MakeShared<FExtender>();
+	ViewMenuExtender->AddMenuExtension("SequencerSettings", EExtensionHook::Before, nullptr, FMenuExtensionDelegate::CreateRaw(this, &FCameraAnimationSequenceCustomization::ExtendViewMenu));
+	Customization.ViewMenuExtender = ViewMenuExtender;
+
 	Customization.OnAssetsDrop.BindLambda([](const TArray<UObject*>&, const FAssetDragDropOp&) -> ESequencerDropResult { return ESequencerDropResult::DropDenied; });
 	Customization.OnClassesDrop.BindLambda([](const TArray<TWeakObjectPtr<UClass>>&, const FClassDragDropOp&) -> ESequencerDropResult { return ESequencerDropResult::DropDenied; });
 	Customization.OnActorsDrop.BindLambda([](const TArray<TWeakObjectPtr<AActor>>&, const FActorDragDropOp&) -> ESequencerDropResult { return ESequencerDropResult::DropDenied; });
@@ -32,19 +38,45 @@ void FCameraAnimationSequenceCustomization::UnregisterSequencerCustomization()
 	FTemplateSequenceCustomizationBase::UnregisterSequencerCustomization();
 }
 
-void FCameraAnimationSequenceCustomization::ExtendSequencerToolbar(FToolBarBuilder& ToolbarBuilder)
+void FCameraAnimationSequenceCustomization::ExtendViewMenu(FMenuBuilder& MenuBuilder)
 {
-	ToolbarBuilder.BeginStyleOverride("SequencerToolBar");
-
-	ToolbarBuilder.AddSeparator();
+	MenuBuilder.BeginSection("CameraAnimations", LOCTEXT("CameraAnimations", "Camera Animations"));
+	{
+		MenuBuilder.AddMenuEntry(
+				LOCTEXT("CameraAdditiveToViewport", "Base Cameras on Viewport"),
+				LOCTEXT("CameraAdditiveToViewportToolTip", "Play additive camera sections on top of viewport position"),
+				FSlateIcon(),
+				FUIAction(
+					FExecuteAction::CreateRaw(this, &FCameraAnimationSequenceCustomization::ToggleCameraInitiallyAdditiveToViewport),
+					FCanExecuteAction(),
+					FIsActionChecked::CreateRaw(this, &FCameraAnimationSequenceCustomization::IsCameraInitiallyAdditiveToViewport)),
+				NAME_None,
+				EUserInterfaceActionType::ToggleButton);
+	}
+	MenuBuilder.EndSection();
+}
 
-	ToolbarBuilder.AddComboButton(
-			FUIAction(),
-			FOnGetContent::CreateRaw(this, &FCameraAnimationSequenceCustomization::GetBoundCameraClassMenuContent),
-			LOCTEXT("CameraTypePicker", "Camera Type"),
-			LOCTEXT("CameraTypePickerTooltip", "Change the base camera actor type that this template sequence can bind to"),
-			FSlateIcon(FTemplateSequenceEditorStyle::Get()->GetStyleSetName(), "TemplateSequenceEditor.Chain"));
-	
+void FCameraAnimationSequenceCustomization::ExtendToolbar(FToolBarBuilder& ToolbarBuilder)
+{
+	ToolbarBuilder.BeginStyleOverride("SequencerToolBar");
+	ToolbarBuilder.BeginSection("CameraAnimations", true);
+	{
+		ToolbarBuilder.AddComboButton(
+				FUIAction(),
+				FOnGetContent::CreateRaw(this, &FCameraAnimationSequenceCustomization::GetBoundCameraClassMenuContent),
+				LOCTEXT("CameraTypePicker", "Camera Type"),
+				LOCTEXT("CameraTypePickerTooltip", "Change the base camera actor type that this template sequence can bind to"),
+				FSlateIcon(FTemplateSequenceEditorStyle::Get()->GetStyleSetName(), "TemplateSequenceEditor.Chain"));
+
+		ToolbarBuilder.AddToolBarButton(
+				FUIAction(
+					FExecuteAction::CreateRaw(this, &FCameraAnimationSequenceCustomization::SetCameraAdditiveToViewportOffset)),
+				NAME_None,
+				FText::GetEmpty(),
+				LOCTEXT("SetCameraAdditiveToViewportOffsetToolTip", "Base Cameras to Current Viewport Transform"),
+				FSlateIcon(FTemplateSequenceEditorStyle::Get()->GetStyleSetName(), "SetViewportPreviewOffset"));
+	}
+	ToolbarBuilder.EndSection();
 	ToolbarBuilder.EndStyleOverride();
 }
 
@@ -81,4 +113,26 @@ bool FCameraAnimationSequenceCustomization::IsBoundToActorClass(UClass* InClass)
 	return BoundActorClass == InClass;
 }
 
+void FCameraAnimationSequenceCustomization::ToggleCameraInitiallyAdditiveToViewport()
+{
+	UTemplateSequenceEditorSettings* Settings = GetMutableDefault<UTemplateSequenceEditorSettings>();
+	Settings->bCameraInitiallyAdditiveToViewport = !Settings->bCameraInitiallyAdditiveToViewport;
+	Settings->SaveConfig();
+}
+
+bool FCameraAnimationSequenceCustomization::IsCameraInitiallyAdditiveToViewport()
+{
+	UTemplateSequenceEditorSettings* Settings = GetMutableDefault<UTemplateSequenceEditorSettings>();
+	return Settings->bCameraInitiallyAdditiveToViewport;
+}
+
+void FCameraAnimationSequenceCustomization::SetCameraAdditiveToViewportOffset()
+{
+	UTemplateSequenceCameraPreviewSystem::EnableNextFrame();
+	if (ISequencer* SequencerPtr = GetSequencer())
+	{
+		SequencerPtr->ForceEvaluate();
+	}
+}
+
 #undef LOCTEXT_NAMESPACE
diff --git a/Engine/Plugins/MovieScene/TemplateSequence/Source/TemplateSequenceEditor/Private/Customizations/CameraAnimationSequenceCustomization.h b/Engine/Plugins/MovieScene/TemplateSequence/Source/TemplateSequenceEditor/Private/Customizations/CameraAnimationSequenceCustomization.h
--- a/Engine/Plugins/MovieScene/TemplateSequence/Source/TemplateSequenceEditor/Private/Customizations/CameraAnimationSequenceCustomization.h
+++ b/Engine/Plugins/MovieScene/TemplateSequence/Source/TemplateSequenceEditor/Private/Customizations/CameraAnimationSequenceCustomization.h
@@ -18,11 +18,16 @@ private:
 	virtual void RegisterSequencerCustomization(FSequencerCustomizationBuilder& Builder) override;
 	virtual void UnregisterSequencerCustomization() override;
 
-	void ExtendSequencerToolbar(FToolBarBuilder& ToolbarBuilder);
+	void ExtendViewMenu(FMenuBuilder& MenuBuilder);
+	void ExtendToolbar(FToolBarBuilder& ToolbarBuilder);
 
 	TSharedRef<SWidget> GetBoundCameraClassMenuContent();
 	bool IsBoundToActorClass(UClass* InClass);
 
+	void ToggleCameraInitiallyAdditiveToViewport();
+	bool IsCameraInitiallyAdditiveToViewport();
+	void SetCameraAdditiveToViewportOffset();
+
 	TArray<UClass*> CameraActorClasses;
 };
 
diff --git a/Engine/Plugins/MovieScene/TemplateSequence/Source/TemplateSequenceEditor/Private/Misc/TemplateSequenceEditorSettings.h b/Engine/Plugins/MovieScene/TemplateSequence/Source/TemplateSequenceEditor/Private/Misc/TemplateSequenceEditorSettings.h
--- a/Engine/Plugins/MovieScene/TemplateSequence/Source/TemplateSequenceEditor/Private/Misc/TemplateSequenceEditorSettings.h
+++ b/Engine/Plugins/MovieScene/TemplateSequence/Source/TemplateSequenceEditor/Private/Misc/TemplateSequenceEditorSettings.h
@@ -20,5 +20,9 @@ public:
 	/** Whether to show "outdated assets" by default in the camera animation track's asset picker. */
 	UPROPERTY(config)
 	bool bShowOutdatedAssetsInCameraAnimationTrackEditor = false;
+
+	/** Whether to set the viewport position as the initial value for additive camera sections. */
+	UPROPERTY(config)
+	bool bCameraInitiallyAdditiveToViewport = true;
 };
 
diff --git a/Engine/Plugins/MovieScene/TemplateSequence/Source/TemplateSequenceEditor/Private/Styles/TemplateSequenceEditorStyle.h b/Engine/Plugins/MovieScene/TemplateSequence/Source/TemplateSequenceEditor/Private/Styles/TemplateSequenceEditorStyle.h
--- a/Engine/Plugins/MovieScene/TemplateSequence/Source/TemplateSequenceEditor/Private/Styles/TemplateSequenceEditorStyle.h
+++ b/Engine/Plugins/MovieScene/TemplateSequence/Source/TemplateSequenceEditor/Private/Styles/TemplateSequenceEditorStyle.h
@@ -34,8 +34,10 @@ public:
 		Set("ClassThumbnail.TemplateSequence", new IMAGE_BRUSH("TemplateSequence_64x", Icon64x64));
 
 		// toolbar icons
-		Set("TemplateSequenceEditor.Chain", new IMAGE_BRUSH("Chain_16x", Icon16x16));
-		Set("TemplateSequenceEditor.Chain.Small", new IMAGE_BRUSH("Chain_24x", Icon24x24));
+		Set("TemplateSequenceEditor.Chain", new IMAGE_BRUSH("Chain_24x", Icon24x24));
+		Set("TemplateSequenceEditor.Chain.Small", new IMAGE_BRUSH("Chain_16x", Icon16x16));
+		Set("SetViewportPreviewOffset", new IMAGE_BRUSH("SetViewportPreviewOffset_24x", Icon24x24));
+		Set("SetViewportPreviewOffset.Small", new IMAGE_BRUSH("SetViewportPreviewOffset_16x", Icon16x16));
 
 		FSlateStyleRegistry::RegisterSlateStyle(*this);
 	}
diff --git a/Engine/Plugins/MovieScene/TemplateSequence/Source/TemplateSequenceEditor/Private/Systems/TemplateSequenceCameraPreviewSystem.cpp b/Engine/Plugins/MovieScene/TemplateSequence/Source/TemplateSequenceEditor/Private/Systems/TemplateSequenceCameraPreviewSystem.cpp
new file mode 100644
--- /dev/null
+++ b/Engine/Plugins/MovieScene/TemplateSequence/Source/TemplateSequenceEditor/Private/Systems/TemplateSequenceCameraPreviewSystem.cpp
@@ -0,0 +1,274 @@
+// Copyright Epic Games, Inc. All Rights Reserved.
+
+#include "Systems/TemplateSequenceCameraPreviewSystem.h"
+
+#include "Camera/CameraActor.h"
+#include "Editor.h"
+#include "Editor/EditorEngine.h"
+#include "EntitySystem/BuiltInComponentTypes.h"
+#include "EntitySystem/Interrogation/MovieSceneInterrogationExtension.h"
+#include "EntitySystem/MovieSceneBlenderSystemTypes.h"
+#include "EntitySystem/MovieSceneEntityMutations.h"
+#include "EntitySystem/MovieSceneEntitySystemLinker.h"
+#include "EntitySystem/MovieSceneEntitySystemRunner.h"
+#include "EntitySystem/MovieSceneEntitySystemTask.h"
+#include "LevelEditorViewport.h"
+#include "Misc/TemplateSequenceEditorSettings.h"
+#include "MovieScene.h"
+#include "MovieSceneBinding.h"
+#include "MovieSceneSequence.h"
+#include "MovieSceneSpawnable.h"
+#include "MovieSceneTracksComponentTypes.h"
+#include "TemplateSequence.h"
+
+#include UE_INLINE_GENERATED_CPP_BY_NAME(TemplateSequenceCameraPreviewSystem)
+
+namespace UE::MovieScene
+{
+
+TEntitySystemLinkerExtensionID<FEditorViewportLinkerExtension> FEditorViewportLinkerExtension::GetExtensionID()
+{
+	static TEntitySystemLinkerExtensionID<FEditorViewportLinkerExtension> ID = UMovieSceneEntitySystemLinker::RegisterExtension<FEditorViewportLinkerExtension>();
+	return ID;
+}
+
+TSharedPtr<FEditorViewportLinkerExtension> FEditorViewportLinkerExtension::GetOrCreateExtension(UMovieSceneEntitySystemLinker* Linker)
+{
+	if (FEditorViewportLinkerExtension* ViewportExtension = Linker->FindExtension<FEditorViewportLinkerExtension>())
+	{
+		return ViewportExtension->AsShared();
+	}
+
+	TSharedPtr<FEditorViewportLinkerExtension> NewViewportExtension = MakeShared<FEditorViewportLinkerExtension>(Linker);
+	Linker->AddExtension(NewViewportExtension.Get());
+	return NewViewportExtension;
+}
+
+FEditorViewportLinkerExtension::FEditorViewportLinkerExtension(UMovieSceneEntitySystemLinker* Linker)
+	: TSharedEntitySystemLinkerExtension(Linker)
+{
+}
+
+struct FEvaluateViewportTransform
+{
+	TOptional<FTransform> ViewportTransform;
+	const FInstanceRegistry* InstanceRegistry;
+	TArray<FMovieSceneEntityID> Entities;
+
+	FEvaluateViewportTransform(const FInstanceRegistry* InInstanceRegistry, const TArray<FMovieSceneEntityID> InEntities)
+		: InstanceRegistry(InInstanceRegistry)
+		, Entities(InEntities)
+	{}
+
+	void PreTask()
+	{
+		ViewportTransform.Reset();
+		if (GEditor)
+		{
+			for (FLevelEditorViewportClient* ViewportClient : GEditor->GetLevelViewportClients())
+			{
+				if (ViewportClient && 
+						ViewportClient->IsPerspective() && 
+						ViewportClient->GetViewMode() != VMI_Unknown &&
+						ViewportClient->AllowsCinematicControl())
+				{
+					ViewportTransform.Emplace(FTransform::Identity);
+					ViewportTransform->SetLocation(ViewportClient->GetViewLocation());
+					ViewportTransform->SetRotation(ViewportClient->GetViewRotation().Quaternion());
+					break;
+				}
+			}
+		}
+	}
+
+	void ForEachAllocation(
+			const FEntityAllocation* Allocation, 
+			FReadEntityIDs EntityIDs,
+			TWrite<FIntermediate3DTransform> InitialTransforms) const
+	{
+		if (!ViewportTransform.IsSet())
+		{
+			return;
+		}
+
+		const int32 Num = Allocation->Num();
+		for (int32 Index = 0; Index < Num; ++Index)
+		{
+			FMovieSceneEntityID EntityID = EntityIDs[Index];
+			if (Entities.Contains(EntityID))
+			{
+				FIntermediate3DTransform& InitialTransform = InitialTransforms[Index];
+				InitialTransform = FIntermediate3DTransform(
+						ViewportTransform->GetLocation(),
+						ViewportTransform->GetRotation().Rotator(),
+						ViewportTransform->GetScale3D());
+			}
+		}
+	}
+};
+
+}
+
+bool UTemplateSequenceCameraPreviewSystem::bEnableNextFrame = false;
+
+void UTemplateSequenceCameraPreviewSystem::EnableNextFrame()
+{
+	bEnableNextFrame = true;
+}
+
+UTemplateSequenceCameraPreviewSystem::UTemplateSequenceCameraPreviewSystem(const FObjectInitializer& ObjInit)
+	: Super(ObjInit)
+{
+	using namespace UE::MovieScene;
+
+	Phase = ESystemPhase::Evaluation;
+}
+
+bool UTemplateSequenceCameraPreviewSystem::IsRelevantImpl(UMovieSceneEntitySystemLinker* InLinker) const
+{
+	return bEnableNextFrame;
+}
+
+void UTemplateSequenceCameraPreviewSystem::OnRun(FSystemTaskPrerequisites& InPrerequisites, FSystemSubsequentTasks& Subsequents)
+{
+	using namespace UE::MovieScene;
+
+	// Don't actually run if we don't really want cameras to be additive to viewports.
+	const UTemplateSequenceEditorSettings* Settings = GetDefault<UTemplateSequenceEditorSettings>();
+	if (!Settings->bCameraInitiallyAdditiveToViewport)
+	{
+		return;
+	}
+
+	if (!bEnableNextFrame)
+	{
+		return;
+	}
+
+	const FBuiltInComponentTypes* BuiltInComponents = FBuiltInComponentTypes::Get();
+	const FMovieSceneTracksComponentTypes* TrackComponents = FMovieSceneTracksComponentTypes::Get();
+
+	FInstanceRegistry* InstanceRegistry = Linker->GetInstanceRegistry();
+
+	// Find all the object/component binding IDs that relate to camera actors.
+	TArray<FGuid> CameraBindingIDs;
+	for (auto It = InstanceRegistry->GetSparseInstances().CreateConstIterator(); It; ++It)
+	{
+		const FSequenceInstance& Instance = (*It);
+		if (!Instance.IsRootSequence())
+		{
+			continue;
+		}
+
+		TSharedRef<FSharedPlaybackState> SharedPlaybackState = Instance.GetSharedPlaybackState();
+		UTemplateSequence* RootSequence = Cast<UTemplateSequence>(SharedPlaybackState->GetRootSequence());
+		if (!RootSequence)
+		{
+			continue;
+		}
+
+		const UObject* RootObjectTemplate = RootSequence->GetRootObjectSpawnableTemplate();
+		if (!RootObjectTemplate || !RootObjectTemplate->IsA<ACameraActor>())
+		{
+			continue;
+		}
+
+		UMovieScene* MovieScene = RootSequence->GetMovieScene();
+		if (!MovieScene || MovieScene->GetSpawnableCount() == 0)
+		{
+			continue;
+		}
+
+		const FMovieSceneSpawnable& FirstSpawnable = MovieScene->GetSpawnable(0);
+		CameraBindingIDs.Add(FirstSpawnable.GetGuid());
+	}
+
+	// Now look for scene component binding entities that were imported from these object bindings.
+	TArray<FMovieSceneEntityID> ParentsOfEntitiesToTag;
+
+	FEntityTaskBuilder()
+	.ReadEntityIDs()
+	.Read(BuiltInComponents->GenericObjectBinding)
+	.FilterAll({ BuiltInComponents->Tags.ImportedEntity })
+	.Iterate_PerEntity(&Linker->EntityManager,
+			[&CameraBindingIDs, &ParentsOfEntitiesToTag](FMovieSceneEntityID EntityID, FGuid ObjectBindingID)
+			{
+				if (CameraBindingIDs.Contains(ObjectBindingID))
+				{
+					ParentsOfEntitiesToTag.Add(EntityID);
+				}
+			});
+
+	// Next, look for transform entities generated from the ones we just found. We are looking for
+	// the entities that hold the initial value of the camera transform. There are two situations:
+	//
+	// 1) They are non-blending entities that handle everything by themselves. Tag those with our
+	//    custom tag.
+	//
+	// 2) They are part of a blending channel's inputs. In that case, we need to find the blending
+	//	  channel output entity, because that's where the initial value is held.
+	//
+	TArray<FMovieSceneEntityID> EntitiesToTag;
+
+	// Situation 1
+	FEntityTaskBuilder()
+	.ReadEntityIDs()
+	.Read(BuiltInComponents->ParentEntity)
+	.FilterAll({
+			TrackComponents->ComponentTransform.PropertyTag,
+			TrackComponents->ComponentTransform.InitialValue })
+	.FilterNone({ BuiltInComponents->BlendChannelInput })
+	.Iterate_PerEntity(&Linker->EntityManager, 
+			[&ParentsOfEntitiesToTag, &EntitiesToTag](
+				FMovieSceneEntityID EntityID, FMovieSceneEntityID ParentEntityID)
+			{
+				if (ParentsOfEntitiesToTag.Contains(ParentEntityID))
+				{
+					EntitiesToTag.Add(EntityID);
+				}
+			});
+
+	// Situation 2
+	TArray<FMovieSceneBlendChannelID> BlendChannelsToTag;
+
+	FEntityTaskBuilder()
+	.Read(BuiltInComponents->ParentEntity)
+	.Read(BuiltInComponents->BlendChannelInput)
+	.Iterate_PerEntity(&Linker->EntityManager, 
+			[&ParentsOfEntitiesToTag, &BlendChannelsToTag](
+				FMovieSceneEntityID ParentEntityID, FMovieSceneBlendChannelID BlendChannelID)
+			{
+				if (ParentsOfEntitiesToTag.Contains(ParentEntityID))
+				{
+					BlendChannelsToTag.Add(BlendChannelID);
+				}
+			});
+
+	FEntityTaskBuilder()
+	.ReadEntityIDs()
+	.Read(BuiltInComponents->BlendChannelOutput)
+	.FilterAll({ 
+			TrackComponents->ComponentTransform.PropertyTag,
+			TrackComponents->ComponentTransform.InitialValue })
+	.Iterate_PerEntity(&Linker->EntityManager,
+			[&BlendChannelsToTag, &EntitiesToTag](FMovieSceneEntityID EntityID, FMovieSceneBlendChannelID BlendChannelID)
+			{
+				 if (BlendChannelsToTag.Contains(BlendChannelID))
+				 {
+					EntitiesToTag.Add(EntityID);
+				 }
+			});
+
+	FEntityTaskBuilder()
+	.ReadEntityIDs()
+	.Write(TrackComponents->ComponentTransform.InitialValue)
+	.FilterAll({ TrackComponents->ComponentTransform.PropertyTag })
+	.Dispatch_PerAllocation<FEvaluateViewportTransform>(
+			&Linker->EntityManager, InPrerequisites, &Subsequents, InstanceRegistry, EntitiesToTag);
+
+	if (!Linker->FindExtension<IInterrogationExtension>())
+	{
+		bEnableNextFrame = false;
+	}
+}
+
diff --git a/Engine/Plugins/MovieScene/TemplateSequence/Source/TemplateSequenceEditor/Private/Systems/TemplateSequenceCameraPreviewSystem.h b/Engine/Plugins/MovieScene/TemplateSequence/Source/TemplateSequenceEditor/Private/Systems/TemplateSequenceCameraPreviewSystem.h
new file mode 100644
--- /dev/null
+++ b/Engine/Plugins/MovieScene/TemplateSequence/Source/TemplateSequenceEditor/Private/Systems/TemplateSequenceCameraPreviewSystem.h
@@ -0,0 +1,43 @@
+// Copyright Epic Games, Inc. All Rights Reserved.
+
+#pragma once
+
+#include "EntitySystem/MovieSceneEntitySystem.h"
+#include "EntitySystem/MovieSceneEntitySystemLinkerSharedExtension.h"
+
+#include "TemplateSequenceCameraPreviewSystem.generated.h"
+
+namespace UE::MovieScene
+{
+
+struct FEditorViewportLinkerExtension : public TSharedEntitySystemLinkerExtension<FEditorViewportLinkerExtension>
+{
+	static TEntitySystemLinkerExtensionID<FEditorViewportLinkerExtension> GetExtensionID();
+	static TSharedPtr<FEditorViewportLinkerExtension> GetOrCreateExtension(UMovieSceneEntitySystemLinker* Linker);
+
+	FEditorViewportLinkerExtension(UMovieSceneEntitySystemLinker* Linker);
+};
+
+}
+
+UCLASS(MinimalAPI)
+class UTemplateSequenceCameraPreviewSystem : public UMovieSceneEntitySystem
+{
+public:
+
+	GENERATED_BODY()
+
+	UTemplateSequenceCameraPreviewSystem(const FObjectInitializer& ObjInit);
+
+	static void EnableNextFrame();
+
+protected:
+
+	virtual bool IsRelevantImpl(UMovieSceneEntitySystemLinker* InLinker) const override;
+	virtual void OnRun(FSystemTaskPrerequisites& InPrerequisites, FSystemSubsequentTasks& Subsequents) override;
+
+private:
+
+	static bool bEnableNextFrame;
+};
+
diff --git a/Engine/Plugins/MovieScene/TemplateSequence/Source/TemplateSequenceEditor/Private/TemplateSequenceEditorToolkit.cpp b/Engine/Plugins/MovieScene/TemplateSequence/Source/TemplateSequenceEditor/Private/TemplateSequenceEditorToolkit.cpp
--- a/Engine/Plugins/MovieScene/TemplateSequence/Source/TemplateSequenceEditor/Private/TemplateSequenceEditorToolkit.cpp
+++ b/Engine/Plugins/MovieScene/TemplateSequence/Source/TemplateSequenceEditor/Private/TemplateSequenceEditorToolkit.cpp
@@ -6,11 +6,13 @@
 #include "LevelEditor.h"
 #include "LevelEditorSequencerIntegration.h"
 #include "Misc/TemplateSequenceEditorPlaybackContext.h"
+#include "Misc/TemplateSequenceEditorSettings.h"
 #include "Misc/TemplateSequenceEditorSpawnRegister.h"
 #include "Modules/ModuleManager.h"
 #include "MovieScene.h"
 #include "ScopedTransaction.h"
 #include "Selection.h"
+#include "Systems/TemplateSequenceCameraPreviewSystem.h"
 #include "Widgets/Docking/SDockTab.h"
 #include "TemplateSequence.h"
 
@@ -119,6 +121,13 @@ void FTemplateSequenceEditorToolkit::Initialize(const EToolkitMode::Type Mode, c
 		}
 	}
 
+	// If we are viewing a camera animation, see if we want to base it on the viewport transform from the start.
+	const UTemplateSequenceEditorSettings* Settings = GetDefault<UTemplateSequenceEditorSettings>();
+	if (Settings->bCameraInitiallyAdditiveToViewport)
+	{
+		UTemplateSequenceCameraPreviewSystem::EnableNextFrame();
+	}
+
 	TSharedPtr<SDockTab> DockTab = LevelEditorModule.AttachSequencer(Sequencer->GetSequencerWidget(), SharedThis(this));
 	if (DockTab.IsValid())
 	{
