Sequencer: Fix shared ptrs in FBXInterop

#rnx

diff --git a/Engine/Plugins/MovieScene/LevelSequenceEditor/Source/LevelSequenceEditor/Private/Misc/LevelSequenceCustomization.cpp b/Engine/Plugins/MovieScene/LevelSequenceEditor/Source/LevelSequenceEditor/Private/Misc/LevelSequenceCustomization.cpp
--- a/Engine/Plugins/MovieScene/LevelSequenceEditor/Source/LevelSequenceEditor/Private/Misc/LevelSequenceCustomization.cpp
+++ b/Engine/Plugins/MovieScene/LevelSequenceEditor/Source/LevelSequenceEditor/Private/Misc/LevelSequenceCustomization.cpp
@@ -96,13 +96,13 @@ void FLevelSequenceCustomization::OnSaveMovieSceneAsClicked()
 
 void FLevelSequenceCustomization::ImportFBX()
 {
-	FLevelSequenceFBXInterop Interop(WeakSequencer.Pin());
+	FLevelSequenceFBXInterop Interop(WeakSequencer.Pin().ToSharedRef());
 	Interop.ImportFBX();
 }
 
 void FLevelSequenceCustomization::ExportFBX()
 {
-	FLevelSequenceFBXInterop Interop(WeakSequencer.Pin());
+	FLevelSequenceFBXInterop Interop(WeakSequencer.Pin().ToSharedRef());
 	Interop.ExportFBX();
 }
 
@@ -247,8 +247,14 @@ void FLevelSequenceCustomization::ExtendObjectBindingContextMenu(FMenuBuilder& M
 		LOCTEXT("ImportFBXTooltip", "Import FBX animation to this object"),
 		FSlateIcon(),
 		FUIAction(
-			FExecuteAction::CreateLambda([=] {
-					FLevelSequenceFBXInterop Interop(EditorViewModel->GetSequencer());
+			FExecuteAction::CreateLambda([this] {
+				const TSharedPtr<ISequencer> Sequencer = WeakSequencer.Pin();
+				if (!Sequencer.IsValid())
+				{
+					return;
+				}
+				
+				FLevelSequenceFBXInterop Interop(Sequencer.ToSharedRef());
 					Interop.ImportFBXOntoSelectedNodes();
 				})
 		));
@@ -258,8 +264,14 @@ void FLevelSequenceCustomization::ExtendObjectBindingContextMenu(FMenuBuilder& M
 		LOCTEXT("ExportFBXTooltip", "Export FBX animation from this object"),
 		FSlateIcon(),
 		FUIAction(
-			FExecuteAction::CreateLambda([=] {
-					FLevelSequenceFBXInterop Interop(EditorViewModel->GetSequencer());
+			FExecuteAction::CreateLambda([this] {
+				const TSharedPtr<ISequencer> Sequencer = WeakSequencer.Pin();
+				if (!Sequencer.IsValid())
+				{
+					return;
+				}
+				
+				FLevelSequenceFBXInterop Interop(Sequencer.ToSharedRef());
 					Interop.ExportFBX();
 				})
 		));
diff --git a/Engine/Plugins/MovieScene/TemplateSequence/Source/TemplateSequenceEditor/Private/Customizations/CameraAnimationSequenceCustomization.cpp b/Engine/Plugins/MovieScene/TemplateSequence/Source/TemplateSequenceEditor/Private/Customizations/CameraAnimationSequenceCustomization.cpp
--- a/Engine/Plugins/MovieScene/TemplateSequence/Source/TemplateSequenceEditor/Private/Customizations/CameraAnimationSequenceCustomization.cpp
+++ b/Engine/Plugins/MovieScene/TemplateSequence/Source/TemplateSequenceEditor/Private/Customizations/CameraAnimationSequenceCustomization.cpp
@@ -129,7 +129,7 @@ bool FCameraAnimationSequenceCustomization::IsCameraInitiallyAdditiveToViewport(
 void FCameraAnimationSequenceCustomization::SetCameraAdditiveToViewportOffset()
 {
 	UTemplateSequenceCameraPreviewSystem::EnableNextFrame();
-	if (ISequencer* SequencerPtr = GetSequencer())
+	if (TSharedPtr<ISequencer> SequencerPtr = GetSequencer())
 	{
 		SequencerPtr->ForceEvaluate();
 	}
diff --git a/Engine/Plugins/MovieScene/TemplateSequence/Source/TemplateSequenceEditor/Private/Customizations/TemplateSequenceCustomizationBase.cpp b/Engine/Plugins/MovieScene/TemplateSequence/Source/TemplateSequenceEditor/Private/Customizations/TemplateSequenceCustomizationBase.cpp
--- a/Engine/Plugins/MovieScene/TemplateSequence/Source/TemplateSequenceEditor/Private/Customizations/TemplateSequenceCustomizationBase.cpp
+++ b/Engine/Plugins/MovieScene/TemplateSequence/Source/TemplateSequenceEditor/Private/Customizations/TemplateSequenceCustomizationBase.cpp
@@ -15,9 +15,11 @@
 
 void FTemplateSequenceCustomizationBase::RegisterSequencerCustomization(FSequencerCustomizationBuilder& Builder)
 {
-	Sequencer = &Builder.GetSequencer();
+	WeakSequencer = Builder.GetSequencer().AsShared();
 	TemplateSequence = Cast<UTemplateSequence>(&Builder.GetFocusedSequence());
 
+	TSharedPtr<ISequencer> Sequencer = WeakSequencer.Pin();
+
 	Sequencer->OnCloseEvent().AddRaw(this, &FTemplateSequenceCustomizationBase::OnSequencerClosed);
 	Sequencer->OnMovieSceneDataChanged().AddRaw(this, &FTemplateSequenceCustomizationBase::OnMovieSceneDataChanged);
 
@@ -29,13 +31,15 @@ void FTemplateSequenceCustomizationBase::RegisterSequencerCustomization(FSequenc
 
 void FTemplateSequenceCustomizationBase::UnregisterSequencerCustomization()
 {
-	if (Sequencer)
+	if (WeakSequencer.IsValid())
 	{
+		TSharedPtr<ISequencer> Sequencer = WeakSequencer.Pin();
+
 		Sequencer->OnCloseEvent().RemoveAll(this);
 		Sequencer->OnMovieSceneDataChanged().RemoveAll(this);
 	}
 
-	Sequencer = nullptr;
+	WeakSequencer = nullptr;
 	TemplateSequence = nullptr;
 }
 
@@ -58,6 +62,7 @@ void FTemplateSequenceCustomizationBase::OnBoundActorClassPicked(UClass* ChosenC
 	{
 		ChangeActorBinding(ChosenClass);
 
+		TSharedPtr<ISequencer> Sequencer = WeakSequencer.Pin();
 		Sequencer->NotifyMovieSceneDataChanged(EMovieSceneDataChangeType::MovieSceneStructureItemAdded);
 	}
 }
@@ -66,6 +71,8 @@ void FTemplateSequenceCustomizationBase::ChangeActorBinding(UObject* Object, UAc
 {
 	const FScopedTransaction Transaction(LOCTEXT("ChangeActorBinding", "Change Template Sequence Actor Binding"));
 
+	TSharedPtr<ISequencer> Sequencer = WeakSequencer.Pin();
+
 	FTemplateSequenceEditorUtil Util(TemplateSequence, *Sequencer);
 	Util.ChangeActorBinding(Object, ActorFactory, bSetupDefaults);
 }
@@ -94,15 +101,19 @@ void FTemplateSequenceCustomizationBase::ExtendObjectBindingContextMenu(FMenuBui
 
 	MenuBuilder.BeginSection("Import/Export", LOCTEXT("ImportExportMenuSectionName", "Import/Export"));
 
-	TSharedPtr<ISequencer> SequencerPtr = Sequencer->AsShared();
-
 	MenuBuilder.AddMenuEntry(
 		LOCTEXT("ImportFBX", "Import..."),
 		LOCTEXT("ImportFBXTooltip", "Import FBX animation to this object"),
 		FSlateIcon(),
 		FUIAction(
-			FExecuteAction::CreateLambda([=] {
-					FLevelSequenceFBXInterop Interop(SequencerPtr);
+			FExecuteAction::CreateLambda([this] {
+				const TSharedPtr<ISequencer> Sequencer = WeakSequencer.Pin();
+				if (!Sequencer.IsValid())
+				{
+					return;
+				}
+
+				FLevelSequenceFBXInterop Interop(Sequencer.ToSharedRef());
 					Interop.ImportFBXOntoSelectedNodes();
 				})
 		));
@@ -112,8 +123,14 @@ void FTemplateSequenceCustomizationBase::ExtendObjectBindingContextMenu(FMenuBui
 		LOCTEXT("ExportFBXTooltip", "Export FBX animation from this object"),
 		FSlateIcon(),
 		FUIAction(
-			FExecuteAction::CreateLambda([=] {
-					FLevelSequenceFBXInterop Interop(SequencerPtr);
+			FExecuteAction::CreateLambda([this] {
+				const TSharedPtr<ISequencer> Sequencer = WeakSequencer.Pin();
+				if (!Sequencer.IsValid())
+				{
+					return;
+				}
+				
+				FLevelSequenceFBXInterop Interop(Sequencer.ToSharedRef());
 					Interop.ExportFBX();
 				})
 		));
@@ -152,12 +169,17 @@ void FTemplateSequenceCustomizationBase::OnMovieSceneDataChanged(EMovieSceneData
 
 void FTemplateSequenceCustomizationBase::OnSequencerClosed(TSharedRef<ISequencer> InSequencer)
 {
-	if (&InSequencer.Get() == Sequencer)
+	if (WeakSequencer.IsValid())
 	{
-		Sequencer->OnCloseEvent().RemoveAll(this);
-		Sequencer->OnMovieSceneDataChanged().RemoveAll(this);
+		TSharedPtr<ISequencer> Sequencer = WeakSequencer.Pin();
+
+		if (InSequencer == Sequencer)
+		{
+			Sequencer->OnCloseEvent().RemoveAll(this);
+			Sequencer->OnMovieSceneDataChanged().RemoveAll(this);
+		}
 
-		Sequencer = nullptr;
+		WeakSequencer = nullptr;
 	}
 }
 
diff --git a/Engine/Plugins/MovieScene/TemplateSequence/Source/TemplateSequenceEditor/Private/Customizations/TemplateSequenceCustomizationBase.h b/Engine/Plugins/MovieScene/TemplateSequence/Source/TemplateSequenceEditor/Private/Customizations/TemplateSequenceCustomizationBase.h
--- a/Engine/Plugins/MovieScene/TemplateSequence/Source/TemplateSequenceEditor/Private/Customizations/TemplateSequenceCustomizationBase.h
+++ b/Engine/Plugins/MovieScene/TemplateSequence/Source/TemplateSequenceEditor/Private/Customizations/TemplateSequenceCustomizationBase.h
@@ -25,7 +25,7 @@ protected:
 	void OnBoundActorClassPicked(UClass* ChosenClass);
 	void ChangeActorBinding(UObject* Object, UActorFactory* ActorFactory = nullptr, bool bSetupDefaults = true);
 
-	ISequencer* GetSequencer() const { return Sequencer; }
+	TSharedPtr<ISequencer> GetSequencer() const { return WeakSequencer.Pin(); }
 	UTemplateSequence* GetTemplateSequence() const { return TemplateSequence; }
 
 	TSharedPtr<FExtender> CreateObjectBindingContextMenuExtender(UE::Sequencer::FViewModelPtr InViewModel);
@@ -36,7 +36,8 @@ private:
 	void OnMovieSceneDataChanged(EMovieSceneDataChangeType ChangeType);
 	void OnSequencerClosed(TSharedRef<ISequencer> InSequencer);
 
-	ISequencer* Sequencer;
+	TWeakPtr<ISequencer> WeakSequencer;
+
 	UTemplateSequence* TemplateSequence;
 };
 
diff --git a/Engine/Source/Editor/MovieSceneTools/Private/LevelSequenceFBXInterop.cpp b/Engine/Source/Editor/MovieSceneTools/Private/LevelSequenceFBXInterop.cpp
--- a/Engine/Source/Editor/MovieSceneTools/Private/LevelSequenceFBXInterop.cpp
+++ b/Engine/Source/Editor/MovieSceneTools/Private/LevelSequenceFBXInterop.cpp
@@ -27,13 +27,19 @@
 
 #define LOCTEXT_NAMESPACE "LevelSequenceFBXInterop"
 
-FLevelSequenceFBXInterop::FLevelSequenceFBXInterop(TSharedPtr<ISequencer> InSequencer)
-	: Sequencer(InSequencer)
+FLevelSequenceFBXInterop::FLevelSequenceFBXInterop(TSharedRef<ISequencer> InSequencer)
+	: WeakSequencer(InSequencer)
 {
 }
 
 void FLevelSequenceFBXInterop::ImportFBX()
 {
+	const TSharedPtr<ISequencer> Sequencer = WeakSequencer.Pin();
+	if (!Sequencer)
+	{
+		return;
+	}
+
 	using namespace UE::Sequencer;
 
 	TMap<FGuid, FString> ObjectBindingNameMap;
@@ -57,6 +63,12 @@ void FLevelSequenceFBXInterop::ImportFBX()
 
 void FLevelSequenceFBXInterop::ImportFBXOntoSelectedNodes()
 {
+	const TSharedPtr<ISequencer> Sequencer = WeakSequencer.Pin();
+	if (!Sequencer)
+	{
+		return;
+	}
+
 	using namespace UE::Sequencer;
 
 	// The object binding and names to match when importing from fbx
@@ -74,6 +86,12 @@ void FLevelSequenceFBXInterop::ImportFBXOntoSelectedNodes()
 
 void FLevelSequenceFBXInterop::ExportFBX()
 {
+	const TSharedPtr<ISequencer> Sequencer = WeakSequencer.Pin();
+	if (!Sequencer)
+	{
+		return;
+	}
+
 	using namespace UE::Sequencer;
 
 	TArray<UExporter*> Exporters;
@@ -203,6 +221,12 @@ void FLevelSequenceFBXInterop::ExportFBX()
 
 void FLevelSequenceFBXInterop::ExportFBXInternal(const FString& ExportFilename, const TArray<FGuid>& Bindings, const TArray<UMovieSceneTrack*>& Tracks)
 {
+	const TSharedPtr<ISequencer> Sequencer = WeakSequencer.Pin();
+	if (!Sequencer)
+	{
+		return;
+	}
+
 	UnFbx::FFbxExporter* Exporter = UnFbx::FFbxExporter::GetInstance();
 	//Show the fbx export dialog options
 	bool ExportCancel = false;
diff --git a/Engine/Source/Editor/MovieSceneTools/Public/LevelSequenceFBXInterop.h b/Engine/Source/Editor/MovieSceneTools/Public/LevelSequenceFBXInterop.h
--- a/Engine/Source/Editor/MovieSceneTools/Public/LevelSequenceFBXInterop.h
+++ b/Engine/Source/Editor/MovieSceneTools/Public/LevelSequenceFBXInterop.h
@@ -8,7 +8,7 @@ class MOVIESCENETOOLS_API FLevelSequenceFBXInterop
 {
 public:
 
-	FLevelSequenceFBXInterop(TSharedPtr<ISequencer> InSequencer);
+	FLevelSequenceFBXInterop(TSharedRef<ISequencer> InSequencer);
 
 	/** Imports the animation from an fbx file. */
 	void ImportFBX();
@@ -24,6 +24,6 @@ private:
 
 private:
 
-	TSharedPtr<ISequencer> Sequencer;
+	TWeakPtr<ISequencer> WeakSequencer;
 };
 
