Sequencer: remove buggy support for weighted blendables in camera animations/shakes

The previous code was trying to save and restore any weighted blendables added by the Sequence onto the camera's PostProcessSettings, in order to combine them with any weighted blendables given by the game camera or some post-FX volume or other. The problem was that this code was also stomping the weight of pre-existing blendables, which prevented the original system that added them to change that weight.

This change removes this code altogether -- we have no support for weighted blendables in Sequencer, so the only way we could add weighted blendables in a camera animation or shake is if someone made a custom 3rd party track for it. Since this is highly unlikely, we are just removing this for now because it's not worth going through the trouble of finding the delta between 2 weighted arrays.

#udn 12002888

diff --git a/Engine/Plugins/MovieScene/TemplateSequence/Source/TemplateSequence/Private/CameraAnimationSequencePlayer.cpp b/Engine/Plugins/MovieScene/TemplateSequence/Source/TemplateSequence/Private/CameraAnimationSequencePlayer.cpp
--- a/Engine/Plugins/MovieScene/TemplateSequence/Source/TemplateSequence/Private/CameraAnimationSequencePlayer.cpp
+++ b/Engine/Plugins/MovieScene/TemplateSequence/Source/TemplateSequence/Private/CameraAnimationSequencePlayer.cpp
@@ -156,10 +156,7 @@ void UCameraAnimationSequenceCameraStandIn::Reset(const FMinimalViewInfo& ViewIn
 
 void UCameraAnimationSequenceCameraStandIn::ResetDefaultValues(const FMinimalViewInfo& ViewInfo)
 {
-	// Save the weighted blendables we want to apply to the camera.
-	TArray<FWeightedBlendable> WBBackup(MoveTemp(PostProcessSettings.WeightedBlendables.Array));
-
-	// We reset all the other properties to the current view's values because a lot of them, like 
+	// We reset all the properties to the current view's values because a lot of them, like 
 	// FieldOfView, don't have any "zero" value that makes sense. We'll figure out the delta in the
 	// update code.
 	bConstrainAspectRatio = ViewInfo.bConstrainAspectRatio;
@@ -171,23 +168,6 @@ void UCameraAnimationSequenceCameraStandIn::ResetDefaultValues(const FMinimalVie
 	// We've set the FieldOfView we have to update the CurrentFocalLength accordingly.
 	CurrentFocalLength = (Filmback.SensorWidth / 2.f) / FMath::Tan(FMath::DegreesToRadians(FieldOfView / 2.f));
 
-	// Restore weighted blendables.
-	// Either the camera animation did not add any blendables, in which case we just restore what
-	// was there before, or it did add some, in which case we restore anything *else* that was
-	// there before. This assumes that there can't be two instances of the same blendable object,
-	// but the AddBlendable method (which we use here) already makes this assumption anyway.
-	if (PostProcessSettings.WeightedBlendables.Array.Num() == 0)
-	{
-		PostProcessSettings.WeightedBlendables.Array = MoveTemp(WBBackup);
-	}
-	else
-	{
-		for (const FWeightedBlendable& WB : WBBackup)
-		{
-			PostProcessSettings.AddBlendable(WB.Object, WB.Weight);
-		}
-	}
-
 	RecalcDerivedData();
 }
 
