Camera animation sequences: small improvements and new icon for the "additive to viewport" feature

- Add the ability to disable the additive offset by resetting all the initial transforms to the origin. This isn't generally correct, but it should be correct for the specific case of camera animations, since the camera actor is a spawnable with no prior state.
- Use new SVG icon for the button that sets the additive offset.
- Move the styles class implementation to the .cpp file.
- Fix a typo in SSequencer that incorrectly tied actions-menu extenders to the view menu.
- Remove some dead/unused code.

diff --git a/Engine/Plugins/MovieScene/TemplateSequence/Source/TemplateSequenceEditor/Private/Customizations/CameraAnimationSequenceCustomization.cpp b/Engine/Plugins/MovieScene/TemplateSequence/Source/TemplateSequenceEditor/Private/Customizations/CameraAnimationSequenceCustomization.cpp
--- a/Engine/Plugins/MovieScene/TemplateSequence/Source/TemplateSequenceEditor/Private/Customizations/CameraAnimationSequenceCustomization.cpp
+++ b/Engine/Plugins/MovieScene/TemplateSequence/Source/TemplateSequenceEditor/Private/Customizations/CameraAnimationSequenceCustomization.cpp
@@ -70,7 +70,8 @@ void FCameraAnimationSequenceCustomization::ExtendToolbar(FToolBarBuilder& Toolb
 
 		ToolbarBuilder.AddToolBarButton(
 				FUIAction(
-					FExecuteAction::CreateRaw(this, &FCameraAnimationSequenceCustomization::SetCameraAdditiveToViewportOffset)),
+					FExecuteAction::CreateRaw(this, &FCameraAnimationSequenceCustomization::SetCameraAdditiveToViewportOffset),
+					FCanExecuteAction::CreateRaw(this, &FCameraAnimationSequenceCustomization::CanSetCameraAdditiveToViewportOffset)),
 				NAME_None,
 				FText::GetEmpty(),
 				LOCTEXT("SetCameraAdditiveToViewportOffsetToolTip", "Base Cameras to Current Viewport Transform"),
@@ -118,6 +119,12 @@ void FCameraAnimationSequenceCustomization::ToggleCameraInitiallyAdditiveToViewp
 	UTemplateSequenceEditorSettings* Settings = GetMutableDefault<UTemplateSequenceEditorSettings>();
 	Settings->bCameraInitiallyAdditiveToViewport = !Settings->bCameraInitiallyAdditiveToViewport;
 	Settings->SaveConfig();
+
+	UTemplateSequenceCameraPreviewSystem::DisableNextFrame();
+	if (TSharedPtr<ISequencer> SequencerPtr = GetSequencer())
+	{
+		SequencerPtr->ForceEvaluate();
+	}
 }
 
 bool FCameraAnimationSequenceCustomization::IsCameraInitiallyAdditiveToViewport()
@@ -135,4 +142,10 @@ void FCameraAnimationSequenceCustomization::SetCameraAdditiveToViewportOffset()
 	}
 }
 
+bool FCameraAnimationSequenceCustomization::CanSetCameraAdditiveToViewportOffset()
+{
+	return IsCameraInitiallyAdditiveToViewport();
+}
+
 #undef LOCTEXT_NAMESPACE
+
diff --git a/Engine/Plugins/MovieScene/TemplateSequence/Source/TemplateSequenceEditor/Private/Customizations/CameraAnimationSequenceCustomization.h b/Engine/Plugins/MovieScene/TemplateSequence/Source/TemplateSequenceEditor/Private/Customizations/CameraAnimationSequenceCustomization.h
--- a/Engine/Plugins/MovieScene/TemplateSequence/Source/TemplateSequenceEditor/Private/Customizations/CameraAnimationSequenceCustomization.h
+++ b/Engine/Plugins/MovieScene/TemplateSequence/Source/TemplateSequenceEditor/Private/Customizations/CameraAnimationSequenceCustomization.h
@@ -26,7 +26,9 @@ private:
 
 	void ToggleCameraInitiallyAdditiveToViewport();
 	bool IsCameraInitiallyAdditiveToViewport();
+
 	void SetCameraAdditiveToViewportOffset();
+	bool CanSetCameraAdditiveToViewportOffset();
 
 	TArray<UClass*> CameraActorClasses;
 };
diff --git a/Engine/Plugins/MovieScene/TemplateSequence/Source/TemplateSequenceEditor/Private/Customizations/TemplateSequenceCustomization.cpp b/Engine/Plugins/MovieScene/TemplateSequence/Source/TemplateSequenceEditor/Private/Customizations/TemplateSequenceCustomization.cpp
--- a/Engine/Plugins/MovieScene/TemplateSequence/Source/TemplateSequenceEditor/Private/Customizations/TemplateSequenceCustomization.cpp
+++ b/Engine/Plugins/MovieScene/TemplateSequence/Source/TemplateSequenceEditor/Private/Customizations/TemplateSequenceCustomization.cpp
@@ -37,8 +37,7 @@ void FTemplateSequenceCustomization::UnregisterSequencerCustomization()
 void FTemplateSequenceCustomization::ExtendSequencerToolbar(FToolBarBuilder& ToolbarBuilder)
 {
 	ToolbarBuilder.BeginStyleOverride("SequencerToolBar");
-
-	ToolbarBuilder.AddSeparator();
+	ToolbarBuilder.BeginSection("TemplateSequence", true);
 
 	ToolbarBuilder.AddComboButton(
 			FUIAction(),
@@ -47,6 +46,7 @@ void FTemplateSequenceCustomization::ExtendSequencerToolbar(FToolBarBuilder& Too
 			LOCTEXT("BoundActorClassPickerTooltip", "Change the base actor type that this template sequence can bind to"),
 			FSlateIcon(FTemplateSequenceEditorStyle::Get()->GetStyleSetName(), "TemplateSequenceEditor.Chain"));
 
+	ToolbarBuilder.EndSection();
 	ToolbarBuilder.EndStyleOverride();
 }
 
diff --git a/Engine/Plugins/MovieScene/TemplateSequence/Source/TemplateSequenceEditor/Private/Styles/TemplateSequenceEditorStyle.cpp b/Engine/Plugins/MovieScene/TemplateSequence/Source/TemplateSequenceEditor/Private/Styles/TemplateSequenceEditorStyle.cpp
--- a/Engine/Plugins/MovieScene/TemplateSequence/Source/TemplateSequenceEditor/Private/Styles/TemplateSequenceEditorStyle.cpp
+++ b/Engine/Plugins/MovieScene/TemplateSequence/Source/TemplateSequenceEditor/Private/Styles/TemplateSequenceEditorStyle.cpp
@@ -2,4 +2,49 @@
 
 #include "Styles/TemplateSequenceEditorStyle.h"
 
+#include "Brushes/SlateImageBrush.h"
+#include "Misc/Paths.h"
+#include "Styling/SlateStyleMacros.h"
+#include "Styling/SlateStyleRegistry.h"
+
 TSharedPtr<FTemplateSequenceEditorStyle> FTemplateSequenceEditorStyle::Singleton;
+
+FTemplateSequenceEditorStyle::FTemplateSequenceEditorStyle()
+	: FSlateStyleSet("TemplateSequenceEditorStyle")
+{
+	const FVector2D Icon16x16(16.0f, 16.0f);
+	const FVector2D Icon24x24(24.0f, 24.0f);
+	const FVector2D Icon48x48(48.0f, 48.0f);
+	const FVector2D Icon64x64(64.0f, 64.0f);
+
+	SetContentRoot(FPaths::EnginePluginsDir() / TEXT("MovieScene/TemplateSequence/Content"));
+
+	// tab icons
+	Set("TemplateSequenceEditor.Tabs.Sequencer", new IMAGE_BRUSH("icon_tab_sequencer_16x", Icon16x16));
+
+	// asset thumbnail
+	Set("ClassThumbnail.TemplateSequence", new IMAGE_BRUSH("TemplateSequence_16x", Icon16x16));
+	Set("ClassThumbnail.TemplateSequence", new IMAGE_BRUSH("TemplateSequence_64x", Icon64x64));
+
+	// toolbar icons
+	Set("TemplateSequenceEditor.Chain", new IMAGE_BRUSH("Chain_24x", Icon24x24));
+	Set("TemplateSequenceEditor.Chain.Small", new IMAGE_BRUSH("Chain_16x", Icon16x16));
+	Set("SetViewportPreviewOffset", new IMAGE_BRUSH_SVG("SetViewportPreviewOffset", Icon24x24));
+	Set("SetViewportPreviewOffset.Small", new IMAGE_BRUSH_SVG("SetViewportPreviewOffset", Icon16x16));
+
+	FSlateStyleRegistry::RegisterSlateStyle(*this);
+}
+
+FTemplateSequenceEditorStyle::~FTemplateSequenceEditorStyle()
+{
+	FSlateStyleRegistry::UnRegisterSlateStyle(*this);
+}
+
+TSharedRef<FTemplateSequenceEditorStyle> FTemplateSequenceEditorStyle::Get()
+{
+	if (!Singleton.IsValid())
+	{
+		Singleton = MakeShareable(new FTemplateSequenceEditorStyle);
+	}
+	return Singleton.ToSharedRef();
+}
diff --git a/Engine/Plugins/MovieScene/TemplateSequence/Source/TemplateSequenceEditor/Private/Styles/TemplateSequenceEditorStyle.h b/Engine/Plugins/MovieScene/TemplateSequence/Source/TemplateSequenceEditor/Private/Styles/TemplateSequenceEditorStyle.h
--- a/Engine/Plugins/MovieScene/TemplateSequence/Source/TemplateSequenceEditor/Private/Styles/TemplateSequenceEditorStyle.h
+++ b/Engine/Plugins/MovieScene/TemplateSequence/Source/TemplateSequenceEditor/Private/Styles/TemplateSequenceEditorStyle.h
@@ -2,13 +2,8 @@
 
 #pragma once
 
-#include "Misc/Paths.h"
-#include "Brushes/SlateImageBrush.h"
-#include "Styling/SlateStyleRegistry.h"
 #include "Styling/SlateStyle.h"
 
-#define IMAGE_BRUSH(RelativePath, ...) FSlateImageBrush(RootToContentDir(RelativePath, TEXT(".png")), __VA_ARGS__)
-
 /**
  * Implements the visual style of the template sequence editor.
  */
@@ -16,45 +11,10 @@ class FTemplateSequenceEditorStyle final : public FSlateStyleSet
 {
 public:
 	
-	FTemplateSequenceEditorStyle()
-		: FSlateStyleSet("TemplateSequenceEditorStyle")
-	{
-		const FVector2D Icon16x16(16.0f, 16.0f);
-		const FVector2D Icon24x24(24.0f, 24.0f);
-		const FVector2D Icon48x48(48.0f, 48.0f);
-		const FVector2D Icon64x64(64.0f, 64.0f);
-
-		SetContentRoot(FPaths::EnginePluginsDir() / TEXT("MovieScene/TemplateSequence/Content"));
-
-		// tab icons
-		Set("TemplateSequenceEditor.Tabs.Sequencer", new IMAGE_BRUSH("icon_tab_sequencer_16x", Icon16x16));
-
-		// asset thumbnail
-		Set("ClassThumbnail.TemplateSequence", new IMAGE_BRUSH("TemplateSequence_16x", Icon16x16));
-		Set("ClassThumbnail.TemplateSequence", new IMAGE_BRUSH("TemplateSequence_64x", Icon64x64));
-
-		// toolbar icons
-		Set("TemplateSequenceEditor.Chain", new IMAGE_BRUSH("Chain_24x", Icon24x24));
-		Set("TemplateSequenceEditor.Chain.Small", new IMAGE_BRUSH("Chain_16x", Icon16x16));
-		Set("SetViewportPreviewOffset", new IMAGE_BRUSH("SetViewportPreviewOffset_24x", Icon24x24));
-		Set("SetViewportPreviewOffset.Small", new IMAGE_BRUSH("SetViewportPreviewOffset_16x", Icon16x16));
-
-		FSlateStyleRegistry::RegisterSlateStyle(*this);
-	}
-
-	virtual ~FTemplateSequenceEditorStyle()
-	{
-		FSlateStyleRegistry::UnRegisterSlateStyle(*this);
-	}
+	FTemplateSequenceEditorStyle();
+	virtual ~FTemplateSequenceEditorStyle();
 
-	static TSharedRef<FTemplateSequenceEditorStyle> Get()
-	{
-		if (!Singleton.IsValid())
-		{
-			Singleton = MakeShareable(new FTemplateSequenceEditorStyle);
-		}
-		return Singleton.ToSharedRef();
-	}
+	static TSharedRef<FTemplateSequenceEditorStyle> Get();
 
 private:
 
diff --git a/Engine/Plugins/MovieScene/TemplateSequence/Source/TemplateSequenceEditor/Private/Systems/TemplateSequenceCameraPreviewSystem.cpp b/Engine/Plugins/MovieScene/TemplateSequence/Source/TemplateSequenceEditor/Private/Systems/TemplateSequenceCameraPreviewSystem.cpp
--- a/Engine/Plugins/MovieScene/TemplateSequence/Source/TemplateSequenceEditor/Private/Systems/TemplateSequenceCameraPreviewSystem.cpp
+++ b/Engine/Plugins/MovieScene/TemplateSequence/Source/TemplateSequenceEditor/Private/Systems/TemplateSequenceCameraPreviewSystem.cpp
@@ -26,43 +26,47 @@
 namespace UE::MovieScene
 {
 
-TEntitySystemLinkerExtensionID<FEditorViewportLinkerExtension> FEditorViewportLinkerExtension::GetExtensionID()
+struct FSetInitialTransformValues
 {
-	static TEntitySystemLinkerExtensionID<FEditorViewportLinkerExtension> ID = UMovieSceneEntitySystemLinker::RegisterExtension<FEditorViewportLinkerExtension>();
-	return ID;
-}
+	TOptional<FTransform> TransformToSet;
+	TArray<FMovieSceneEntityID> Entities;
 
-TSharedPtr<FEditorViewportLinkerExtension> FEditorViewportLinkerExtension::GetOrCreateExtension(UMovieSceneEntitySystemLinker* Linker)
-{
-	if (FEditorViewportLinkerExtension* ViewportExtension = Linker->FindExtension<FEditorViewportLinkerExtension>())
+	void ForEachAllocation(
+			const FEntityAllocation* Allocation, 
+			FReadEntityIDs EntityIDs,
+			TWrite<FIntermediate3DTransform> InitialTransforms) const
 	{
-		return ViewportExtension->AsShared();
-	}
-
-	TSharedPtr<FEditorViewportLinkerExtension> NewViewportExtension = MakeShared<FEditorViewportLinkerExtension>(Linker);
-	Linker->AddExtension(NewViewportExtension.Get());
-	return NewViewportExtension;
-}
+		if (!TransformToSet.IsSet())
+		{
+			return;
+		}
 
-FEditorViewportLinkerExtension::FEditorViewportLinkerExtension(UMovieSceneEntitySystemLinker* Linker)
-	: TSharedEntitySystemLinkerExtension(Linker)
-{
-}
+		const int32 Num = Allocation->Num();
+		for (int32 Index = 0; Index < Num; ++Index)
+		{
+			FMovieSceneEntityID EntityID = EntityIDs[Index];
+			if (Entities.Contains(EntityID))
+			{
+				FIntermediate3DTransform& InitialTransform = InitialTransforms[Index];
+				InitialTransform = FIntermediate3DTransform(
+						TransformToSet->GetLocation(),
+						TransformToSet->GetRotation().Rotator(),
+						TransformToSet->GetScale3D());
+			}
+		}
+	}
+};
 
-struct FEvaluateViewportTransform
+struct FSetInitialTransformToViewport : FSetInitialTransformValues
 {
-	TOptional<FTransform> ViewportTransform;
-	const FInstanceRegistry* InstanceRegistry;
-	TArray<FMovieSceneEntityID> Entities;
-
-	FEvaluateViewportTransform(const FInstanceRegistry* InInstanceRegistry, const TArray<FMovieSceneEntityID> InEntities)
-		: InstanceRegistry(InInstanceRegistry)
-		, Entities(InEntities)
-	{}
+	FSetInitialTransformToViewport(const TArray<FMovieSceneEntityID> InEntities)
+	{
+		Entities = InEntities;
+	}
 
 	void PreTask()
 	{
-		ViewportTransform.Reset();
+		TransformToSet.Reset();
 		if (GEditor)
 		{
 			for (FLevelEditorViewportClient* ViewportClient : GEditor->GetLevelViewportClients())
@@ -72,50 +76,40 @@ struct FEvaluateViewportTransform
 						ViewportClient->GetViewMode() != VMI_Unknown &&
 						ViewportClient->AllowsCinematicControl())
 				{
-					ViewportTransform.Emplace(FTransform::Identity);
-					ViewportTransform->SetLocation(ViewportClient->GetViewLocation());
-					ViewportTransform->SetRotation(ViewportClient->GetViewRotation().Quaternion());
+					TransformToSet.Emplace(FTransform::Identity);
+					TransformToSet->SetLocation(ViewportClient->GetViewLocation());
+					TransformToSet->SetRotation(ViewportClient->GetViewRotation().Quaternion());
 					break;
 				}
 			}
 		}
 	}
+};
 
-	void ForEachAllocation(
-			const FEntityAllocation* Allocation, 
-			FReadEntityIDs EntityIDs,
-			TWrite<FIntermediate3DTransform> InitialTransforms) const
+struct FSetInitialTransformToOrigin : FSetInitialTransformValues
+{
+	FSetInitialTransformToOrigin(const TArray<FMovieSceneEntityID> InEntities)
 	{
-		if (!ViewportTransform.IsSet())
-		{
-			return;
-		}
-
-		const int32 Num = Allocation->Num();
-		for (int32 Index = 0; Index < Num; ++Index)
-		{
-			FMovieSceneEntityID EntityID = EntityIDs[Index];
-			if (Entities.Contains(EntityID))
-			{
-				FIntermediate3DTransform& InitialTransform = InitialTransforms[Index];
-				InitialTransform = FIntermediate3DTransform(
-						ViewportTransform->GetLocation(),
-						ViewportTransform->GetRotation().Rotator(),
-						ViewportTransform->GetScale3D());
-			}
-		}
+		Entities = InEntities;
+		TransformToSet = FTransform::Identity;
 	}
 };
 
 }
 
 bool UTemplateSequenceCameraPreviewSystem::bEnableNextFrame = false;
+bool UTemplateSequenceCameraPreviewSystem::bDisableNextFrame = false;
 
 void UTemplateSequenceCameraPreviewSystem::EnableNextFrame()
 {
 	bEnableNextFrame = true;
 }
 
+void UTemplateSequenceCameraPreviewSystem::DisableNextFrame()
+{
+	bDisableNextFrame = true;
+}
+
 UTemplateSequenceCameraPreviewSystem::UTemplateSequenceCameraPreviewSystem(const FObjectInitializer& ObjInit)
 	: Super(ObjInit)
 {
@@ -126,21 +120,14 @@ UTemplateSequenceCameraPreviewSystem::UTemplateSequenceCameraPreviewSystem(const
 
 bool UTemplateSequenceCameraPreviewSystem::IsRelevantImpl(UMovieSceneEntitySystemLinker* InLinker) const
 {
-	return bEnableNextFrame;
+	return bEnableNextFrame || bDisableNextFrame;
 }
 
 void UTemplateSequenceCameraPreviewSystem::OnRun(FSystemTaskPrerequisites& InPrerequisites, FSystemSubsequentTasks& Subsequents)
 {
 	using namespace UE::MovieScene;
 
-	// Don't actually run if we don't really want cameras to be additive to viewports.
-	const UTemplateSequenceEditorSettings* Settings = GetDefault<UTemplateSequenceEditorSettings>();
-	if (!Settings->bCameraInitiallyAdditiveToViewport)
-	{
-		return;
-	}
-
-	if (!bEnableNextFrame)
+	if (!(bEnableNextFrame || bDisableNextFrame))
 	{
 		return;
 	}
@@ -259,16 +246,29 @@ void UTemplateSequenceCameraPreviewSystem::OnRun(FSystemTaskPrerequisites& InPre
 				 }
 			});
 
-	FEntityTaskBuilder()
-	.ReadEntityIDs()
-	.Write(TrackComponents->ComponentTransform.InitialValue)
-	.FilterAll({ TrackComponents->ComponentTransform.PropertyTag })
-	.Dispatch_PerAllocation<FEvaluateViewportTransform>(
-			&Linker->EntityManager, InPrerequisites, &Subsequents, InstanceRegistry, EntitiesToTag);
+	if (bEnableNextFrame)
+	{
+		FEntityTaskBuilder()
+		.ReadEntityIDs()
+		.Write(TrackComponents->ComponentTransform.InitialValue)
+		.FilterAll({ TrackComponents->ComponentTransform.PropertyTag })
+		.Dispatch_PerAllocation<FSetInitialTransformToViewport>(
+				&Linker->EntityManager, InPrerequisites, &Subsequents, EntitiesToTag);
+	}
+	else if (bDisableNextFrame)
+	{
+		FEntityTaskBuilder()
+		.ReadEntityIDs()
+		.Write(TrackComponents->ComponentTransform.InitialValue)
+		.FilterAll({ TrackComponents->ComponentTransform.PropertyTag })
+		.Dispatch_PerAllocation<FSetInitialTransformToOrigin>(
+				&Linker->EntityManager, InPrerequisites, &Subsequents, EntitiesToTag);
+	}
 
 	if (!Linker->FindExtension<IInterrogationExtension>())
 	{
 		bEnableNextFrame = false;
+		bDisableNextFrame = false;
 	}
 }
 
diff --git a/Engine/Plugins/MovieScene/TemplateSequence/Source/TemplateSequenceEditor/Private/Systems/TemplateSequenceCameraPreviewSystem.h b/Engine/Plugins/MovieScene/TemplateSequence/Source/TemplateSequenceEditor/Private/Systems/TemplateSequenceCameraPreviewSystem.h
--- a/Engine/Plugins/MovieScene/TemplateSequence/Source/TemplateSequenceEditor/Private/Systems/TemplateSequenceCameraPreviewSystem.h
+++ b/Engine/Plugins/MovieScene/TemplateSequence/Source/TemplateSequenceEditor/Private/Systems/TemplateSequenceCameraPreviewSystem.h
@@ -7,19 +7,10 @@
 
 #include "TemplateSequenceCameraPreviewSystem.generated.h"
 
-namespace UE::MovieScene
-{
-
-struct FEditorViewportLinkerExtension : public TSharedEntitySystemLinkerExtension<FEditorViewportLinkerExtension>
-{
-	static TEntitySystemLinkerExtensionID<FEditorViewportLinkerExtension> GetExtensionID();
-	static TSharedPtr<FEditorViewportLinkerExtension> GetOrCreateExtension(UMovieSceneEntitySystemLinker* Linker);
-
-	FEditorViewportLinkerExtension(UMovieSceneEntitySystemLinker* Linker);
-};
-
-}
-
+/**
+ * A sequencer system that can write the current editor viewport transform into the initial values
+ * of the camera object's transform entities.
+ */
 UCLASS(MinimalAPI)
 class UTemplateSequenceCameraPreviewSystem : public UMovieSceneEntitySystem
 {
@@ -29,7 +20,10 @@ public:
 
 	UTemplateSequenceCameraPreviewSystem(const FObjectInitializer& ObjInit);
 
+	/** Update the root camera's initial transform to the current editor viewport transform. */
 	static void EnableNextFrame();
+	/** Revert the root camera's initial transform to the origin. */
+	static void DisableNextFrame();
 
 protected:
 
@@ -39,5 +33,6 @@ protected:
 private:
 
 	static bool bEnableNextFrame;
+	static bool bDisableNextFrame;
 };
 
diff --git a/Engine/Source/Editor/Sequencer/Private/SSequencer.cpp b/Engine/Source/Editor/Sequencer/Private/SSequencer.cpp
--- a/Engine/Source/Editor/Sequencer/Private/SSequencer.cpp
+++ b/Engine/Source/Editor/Sequencer/Private/SSequencer.cpp
@@ -1981,7 +1981,7 @@ TSharedRef<SWidget> SSequencer::MakeViewMenu()
 	TArray<TSharedPtr<FExtender>> AllExtenders;
 	ISequencerModule& SequencerModule = FModuleManager::GetModuleChecked<ISequencerModule>("Sequencer");
 	AllExtenders.Add(SequencerModule.GetViewMenuExtensibilityManager()->GetAllExtenders());
-	AllExtenders.Append(ActionsMenuExtenders);
+	AllExtenders.Append(ViewMenuExtenders);
 
 	TSharedPtr<FExtender> Extender = FExtender::Combine(AllExtenders);
 	FMenuBuilder MenuBuilder(true, Sequencer->GetCommandBindings(), Extender);
