Sequencer: handle the case of changing a nested template sequence's bound actor class

The editor code wasn't doing the correct thing when editing a template sequence inside a template sequence.

The runtime code was getting sequence IDs through the instance registry, but that doesn't necessarily exist when tracks are interrogated, such as for drawing trajectories.

diff --git a/Engine/Plugins/MovieScene/TemplateSequence/Source/TemplateSequence/Private/Section/TemplateSequenceSection.cpp b/Engine/Plugins/MovieScene/TemplateSequence/Source/TemplateSequence/Private/Section/TemplateSequenceSection.cpp
--- a/Engine/Plugins/MovieScene/TemplateSequence/Source/TemplateSequence/Private/Section/TemplateSequenceSection.cpp
+++ b/Engine/Plugins/MovieScene/TemplateSequence/Source/TemplateSequence/Private/Section/TemplateSequenceSection.cpp
@@ -117,8 +117,8 @@ void UTemplateSequenceSection::ImportEntityImpl(UMovieSceneEntitySystemLinker* E
 		return;
 	}
 
-	const FSubSequencePath PathToRoot = EntityLinker->GetInstanceRegistry()->GetInstance(Params.Sequence.InstanceHandle).GetSubSequencePath();
-	const FMovieSceneSequenceID ResolvedSequenceID = PathToRoot.ResolveChildSequenceID(GetSequenceID());
+	const FMovieSceneSequenceID SubSequenceID = GetSequenceID();
+	const FMovieSceneSequenceID ResolvedSequenceID = SubSequenceID.AccumulateParentID(Params.Sequence.SequenceID);
 
 	if (Params.EntityID == 0)
 	{
diff --git a/Engine/Plugins/MovieScene/TemplateSequence/Source/TemplateSequenceEditor/Private/Misc/TemplateSequenceEditorUtil.cpp b/Engine/Plugins/MovieScene/TemplateSequence/Source/TemplateSequenceEditor/Private/Misc/TemplateSequenceEditorUtil.cpp
--- a/Engine/Plugins/MovieScene/TemplateSequence/Source/TemplateSequenceEditor/Private/Misc/TemplateSequenceEditorUtil.cpp
+++ b/Engine/Plugins/MovieScene/TemplateSequence/Source/TemplateSequenceEditor/Private/Misc/TemplateSequenceEditorUtil.cpp
@@ -20,16 +20,30 @@ void FTemplateSequenceEditorUtil::ChangeActorBinding(UObject* Object, UActorFact
 		return;
 	}
 
+	// Get the currently edited sequence. It isn't necessarily TemplateSequence, as we could have template 
+	// sequences inside template sequences.
+	UMovieSceneSequence* FocusedSequence = Sequencer.GetFocusedMovieSceneSequence();
+	if (!ensure(FocusedSequence))
+	{
+		return;
+	}
+
+	UTemplateSequence* FocusedTemplateSequence = Cast<UTemplateSequence>(FocusedSequence);
+	if (!FocusedTemplateSequence)
+	{
+		return;
+	}
+
 	// See if we have anything to do in the first place.
 	if (UClass* ChosenClass = Cast<UClass>(Object))
 	{
-		if (ChosenClass == TemplateSequence->BoundActorClass)
+		if (ChosenClass == FocusedTemplateSequence->BoundActorClass)
 		{
 			return;
 		}
 	}
 
-	UMovieScene* MovieScene = TemplateSequence->GetMovieScene();
+	UMovieScene* MovieScene = FocusedTemplateSequence->GetMovieScene();
 	check(MovieScene != nullptr);
 
 	// See if we previously had a main object binding.
@@ -49,20 +63,20 @@ void FTemplateSequenceEditorUtil::ChangeActorBinding(UObject* Object, UActorFact
 
 	// Pre-notify that a change will occur.
 	FProperty* BoundActorClassProperty = FindFProperty<FProperty>(UTemplateSequence::StaticClass(), GET_MEMBER_NAME_CHECKED(UTemplateSequence, BoundActorClass));
-	TemplateSequence->PreEditChange(BoundActorClassProperty);
+	FocusedTemplateSequence->PreEditChange(BoundActorClassProperty);
 
 	// Update the BoundActorClass.
 	if (Object->IsA<UClass>())
 	{
 		UClass* ChosenClass = StaticCast<UClass*>(Object);
-		TemplateSequence->BoundActorClass = ChosenClass;
+		FocusedTemplateSequence->BoundActorClass = ChosenClass;
 	}
 	else
 	{
 		const UObject* SpawnableTemplate = NewSpawnable->GetObjectTemplate();
 		if (ensure(SpawnableTemplate != nullptr))
 		{
-			TemplateSequence->BoundActorClass = SpawnableTemplate->GetClass();
+			FocusedTemplateSequence->BoundActorClass = SpawnableTemplate->GetClass();
 		}
 	}
 
@@ -80,6 +94,6 @@ void FTemplateSequenceEditorUtil::ChangeActorBinding(UObject* Object, UActorFact
 
 	// Notify that the change occured.
 	FPropertyChangedEvent PropertyEvent(BoundActorClassProperty, EPropertyChangeType::ValueSet);
-	TemplateSequence->PostEditChangeProperty(PropertyEvent);
+	FocusedTemplateSequence->PostEditChangeProperty(PropertyEvent);
 }
 
