Sequencer: revert template sections to old entity import code

This will need another fix for nested template sequences

diff --git a/Engine/Plugins/MovieScene/TemplateSequence/Source/TemplateSequence/Private/Section/TemplateSequenceSection.cpp b/Engine/Plugins/MovieScene/TemplateSequence/Source/TemplateSequence/Private/Section/TemplateSequenceSection.cpp
--- a/Engine/Plugins/MovieScene/TemplateSequence/Source/TemplateSequence/Private/Section/TemplateSequenceSection.cpp
+++ b/Engine/Plugins/MovieScene/TemplateSequence/Source/TemplateSequence/Private/Section/TemplateSequenceSection.cpp
@@ -117,8 +117,8 @@ void UTemplateSequenceSection::ImportEntityImpl(UMovieSceneEntitySystemLinker* E
 		return;
 	}
 
-	const FMovieSceneSequenceID SubSequenceID = GetSequenceID();
-	const FMovieSceneSequenceID ResolvedSequenceID = SubSequenceID.AccumulateParentID(Params.Sequence.SequenceID);
+	const FSubSequencePath PathToRoot = EntityLinker->GetInstanceRegistry()->GetInstance(Params.Sequence.InstanceHandle).GetSubSequencePath();
+	const FMovieSceneSequenceID ResolvedSequenceID = PathToRoot.ResolveChildSequenceID(GetSequenceID());
 
 	if (Params.EntityID == 0)
 	{
