Template Sequence: Check for valid instance handle before getting the subsequence path

diff --git a/Engine/Plugins/MovieScene/TemplateSequence/Source/TemplateSequence/Private/Section/TemplateSequenceSection.cpp b/Engine/Plugins/MovieScene/TemplateSequence/Source/TemplateSequence/Private/Section/TemplateSequenceSection.cpp
--- a/Engine/Plugins/MovieScene/TemplateSequence/Source/TemplateSequence/Private/Section/TemplateSequenceSection.cpp
+++ b/Engine/Plugins/MovieScene/TemplateSequence/Source/TemplateSequence/Private/Section/TemplateSequenceSection.cpp
@@ -117,6 +117,11 @@ void UTemplateSequenceSection::ImportEntityImpl(UMovieSceneEntitySystemLinker* E
 		return;
 	}
 
+	if (!EntityLinker->GetInstanceRegistry()->IsHandleValid(Params.Sequence.InstanceHandle))
+	{
+		return;
+	}
+
 	const FSubSequencePath PathToRoot = EntityLinker->GetInstanceRegistry()->GetInstance(Params.Sequence.InstanceHandle).GetSubSequencePath();
 	const FMovieSceneSequenceID ResolvedSequenceID = PathToRoot.ResolveChildSequenceID(GetSequenceID());
 
