Sequencer: Check UObjectInitialized before StaticClass to fix shutdown stompmalloc errors

#rnx

diff --git a/Engine/Plugins/MovieScene/LevelSequenceEditor/Source/LevelSequenceEditor/Private/LevelSequenceEditorModule.cpp b/Engine/Plugins/MovieScene/LevelSequenceEditor/Source/LevelSequenceEditor/Private/LevelSequenceEditorModule.cpp
--- a/Engine/Plugins/MovieScene/LevelSequenceEditor/Source/LevelSequenceEditor/Private/LevelSequenceEditorModule.cpp
+++ b/Engine/Plugins/MovieScene/LevelSequenceEditor/Source/LevelSequenceEditor/Private/LevelSequenceEditorModule.cpp
@@ -381,8 +381,11 @@ protected:
 
 	void UnregisterSequenceCustomizations()
 	{
-		ISequencerModule& SequencerModule = FModuleManager::LoadModuleChecked<ISequencerModule>("Sequencer");
-		SequencerModule.GetSequencerCustomizationManager()->UnregisterInstancedSequencerCustomization(ULevelSequence::StaticClass());
+		if (UObjectInitialized())
+		{
+			ISequencerModule& SequencerModule = FModuleManager::LoadModuleChecked<ISequencerModule>("Sequencer");
+			SequencerModule.GetSequencerCustomizationManager()->UnregisterInstancedSequencerCustomization(ULevelSequence::StaticClass());
+		}
 	}
 
 protected:
diff --git a/Engine/Plugins/MovieScene/TemplateSequence/Source/TemplateSequence/Private/CameraAnimationSequencePlayer.cpp b/Engine/Plugins/MovieScene/TemplateSequence/Source/TemplateSequence/Private/CameraAnimationSequencePlayer.cpp
--- a/Engine/Plugins/MovieScene/TemplateSequence/Source/TemplateSequence/Private/CameraAnimationSequencePlayer.cpp
+++ b/Engine/Plugins/MovieScene/TemplateSequence/Source/TemplateSequence/Private/CameraAnimationSequencePlayer.cpp
@@ -219,7 +219,7 @@ void UCameraAnimationSequenceCameraStandIn::UnregisterCameraStandIn()
 {
 	using namespace UE::MovieScene;
 
-	if (ensure(bRegistered))
+	if (ensure(bRegistered) && UObjectInitialized())
 	{
 		FMovieSceneTracksComponentTypes* TracksComponentTypes = FMovieSceneTracksComponentTypes::Get();
 		TracksComponentTypes->Accessors.ComponentTransform.RemoveAll(UCameraAnimationSequenceCameraStandIn::StaticClass());
diff --git a/Engine/Plugins/MovieScene/TemplateSequence/Source/TemplateSequenceEditor/Private/TemplateSequenceEditorModule.cpp b/Engine/Plugins/MovieScene/TemplateSequence/Source/TemplateSequenceEditor/Private/TemplateSequenceEditorModule.cpp
--- a/Engine/Plugins/MovieScene/TemplateSequence/Source/TemplateSequenceEditor/Private/TemplateSequenceEditorModule.cpp
+++ b/Engine/Plugins/MovieScene/TemplateSequence/Source/TemplateSequenceEditor/Private/TemplateSequenceEditorModule.cpp
@@ -135,9 +135,12 @@ private:
 
 	void UnregisterSequenceCustomizations()
 	{
-		ISequencerModule& SequencerModule = FModuleManager::LoadModuleChecked<ISequencerModule>("Sequencer");
-		SequencerModule.GetSequencerCustomizationManager()->UnregisterInstancedSequencerCustomization(UCameraAnimationSequence::StaticClass());
-		SequencerModule.GetSequencerCustomizationManager()->UnregisterInstancedSequencerCustomization(UTemplateSequence::StaticClass());
+		if (UObjectInitialized())
+		{
+			ISequencerModule& SequencerModule = FModuleManager::LoadModuleChecked<ISequencerModule>("Sequencer");
+			SequencerModule.GetSequencerCustomizationManager()->UnregisterInstancedSequencerCustomization(UCameraAnimationSequence::StaticClass());
+			SequencerModule.GetSequencerCustomizationManager()->UnregisterInstancedSequencerCustomization(UTemplateSequence::StaticClass());
+		}
 	}
 
 private:
