Sequencer: Initialize HierarchicalBias in PostInitProperties in case it is overwritten by config

- Tested that the user changing the Hierarchical Bais on a Template Sequence Section will properly set and restore the value that the user set.
- Tested that having config set to 100 doesn't set the Template Sequence Section to 100 (it gets correctly initialized to 0).

#Synced-CL 49071976

diff --git a/Engine/Plugins/MovieScene/TemplateSequence/Source/TemplateSequence/Private/Section/TemplateSequenceSection.cpp b/Engine/Plugins/MovieScene/TemplateSequence/Source/TemplateSequence/Private/Section/TemplateSequenceSection.cpp
--- a/Engine/Plugins/MovieScene/TemplateSequence/Source/TemplateSequence/Private/Section/TemplateSequenceSection.cpp
+++ b/Engine/Plugins/MovieScene/TemplateSequence/Source/TemplateSequence/Private/Section/TemplateSequenceSection.cpp
@@ -19,10 +19,15 @@ UTemplateSequenceSection::UTemplateSequenceSection(const FObjectInitializer& Obj
 	: Super(ObjInitializer)
 {
 	SetBlendType(EMovieSceneBlendType::Absolute);
+}
 
+void UTemplateSequenceSection::PostInitProperties()
+{
 	// Template sequences always adopt the same hierarchical bias as their parent sequence so that their
 	// animation can blend with any complementary animation set directly on their target object.
 	Parameters.HierarchicalBias = 0;
+	
+	Super::PostInitProperties();
 }
 
 void UTemplateSequenceSection::AddPropertyScale(const FTemplateSectionPropertyScale& InPropertyScale)
diff --git a/Engine/Plugins/MovieScene/TemplateSequence/Source/TemplateSequence/Public/Sections/TemplateSequenceSection.h b/Engine/Plugins/MovieScene/TemplateSequence/Source/TemplateSequence/Public/Sections/TemplateSequenceSection.h
--- a/Engine/Plugins/MovieScene/TemplateSequence/Source/TemplateSequence/Public/Sections/TemplateSequenceSection.h
+++ b/Engine/Plugins/MovieScene/TemplateSequence/Source/TemplateSequence/Public/Sections/TemplateSequenceSection.h
@@ -73,6 +73,9 @@ protected:
 	UE_API virtual bool PopulateEvaluationFieldImpl(const TRange<FFrameNumber>& EffectiveRange, const FMovieSceneEvaluationFieldEntityMetaData& InMetaData, FMovieSceneEntityComponentFieldBuilder* OutFieldBuilder) override;
 	UE_API virtual void ImportEntityImpl(UMovieSceneEntitySystemLinker* EntityLinker, const UE::MovieScene::FEntityImportParams& Params, UE::MovieScene::FImportedEntity* OutImportedEntity) override;
 
+	/**~ UObject implementation */
+	UE_API virtual void PostInitProperties() override;
+
 public:
 
 	UPROPERTY()
